XRScribe: alpha implementation in Plugins/Runtime/XR

diff --git a/Engine/Plugins/Runtime/XR/XRScribe/Source/XRScribe/Private/XRScribeDeveloperSettings.cpp b/Engine/Plugins/Runtime/XR/XRScribe/Source/XRScribe/Private/XRScribeDeveloperSettings.cpp
new file mode 100644
--- /dev/null
+++ b/Engine/Plugins/Runtime/XR/XRScribe/Source/XRScribe/Private/XRScribeDeveloperSettings.cpp
@@ -0,0 +1,10 @@
+// Copyright Epic Games, Inc. All Rights Reserved.
+
+#include "XRScribeDeveloperSettings.h"
+
+//#include UE_INLINE_GENERATED_CPP_BY_NAME(XRScribeDeveloperSettings)
+//
+//UXRScribeDeveloperSettings::UXRScribeDeveloperSettings(const FObjectInitializer& Initializer)
+//	: Super(Initializer)
+//{
+//}
diff --git a/Engine/Plugins/Runtime/XR/XRScribe/Source/XRScribe/Private/XRScribeEmulatedPoseManager.cpp b/Engine/Plugins/Runtime/XR/XRScribe/Source/XRScribe/Private/XRScribeEmulatedPoseManager.cpp
new file mode 100644
--- /dev/null
+++ b/Engine/Plugins/Runtime/XR/XRScribe/Source/XRScribe/Private/XRScribeEmulatedPoseManager.cpp
@@ -0,0 +1,162 @@
+// Copyright Epic Games, Inc. All Rights Reserved.
+
+#include "XRScribeEmulatedPoseManager.h"
+#include "XRScribeFileFormat.h"
+
+namespace UE::XRScribe
+{
+
+XrReferenceSpaceType FOpenXRPoseManager::GetCapturedReferenceSpaceType(XrSpace CapturedSpace)
+{
+	XrReferenceSpaceType Result = XR_REFERENCE_SPACE_TYPE_MAX_ENUM;
+
+	FOpenXRCreateReferenceSpacePacket* MatchingInfo = CapturedReferenceSpaces.FindByPredicate([CapturedSpace](const FOpenXRCreateReferenceSpacePacket& CreateRefSpaceInfo)
+	{
+		return CreateRefSpaceInfo.Space == CapturedSpace;
+	});
+
+	if (MatchingInfo != nullptr)
+	{
+		Result = MatchingInfo->ReferenceSpaceCreateInfo.referenceSpaceType;
+	}
+
+	return Result;
+}
+
+void FOpenXRPoseManager::RegisterCapturedReferenceSpace(const FOpenXRCreateReferenceSpacePacket& CreateReferenceSpacePacket)
+{
+	CapturedReferenceSpaces.Add(CreateReferenceSpacePacket);
+
+	// TODO: check for doubles?
+
+	if ((CreateReferenceSpacePacket.Result == XR_SUCCESS) &&
+		(CreateReferenceSpacePacket.ReferenceSpaceCreateInfo.referenceSpaceType == XR_REFERENCE_SPACE_TYPE_STAGE) &&
+		TrackingSpaceType != XR_REFERENCE_SPACE_TYPE_STAGE)
+	{
+		TrackingSpaceType = XR_REFERENCE_SPACE_TYPE_STAGE;
+	}
+}
+
+void FOpenXRPoseManager::RegisterCapturedSpaceHistory(const TArray<FOpenXRLocateSpacePacket>& SpaceHistory)
+{
+	if (SpaceHistory.Num() == 0)
+	{
+		// error info
+		return;
+	}
+
+	XrSpace CapturedLocatedSpace = SpaceHistory[0].Space;
+	XrReferenceSpaceType RefSpaceType = GetCapturedReferenceSpaceType(CapturedLocatedSpace);
+
+	if (RefSpaceType == XR_REFERENCE_SPACE_TYPE_MAX_ENUM)
+	{
+		return;
+	}
+
+	if (RefSpaceType != XR_REFERENCE_SPACE_TYPE_VIEW)
+	{
+		return;
+	}
+
+	XrSpace CapturedBaseSpace = SpaceHistory[0].BaseSpace;
+	XrReferenceSpaceType BaseRefSpaceType = GetCapturedReferenceSpaceType(CapturedBaseSpace);
+
+	if (BaseRefSpaceType != XR_REFERENCE_SPACE_TYPE_STAGE)
+	{
+		return;
+	}
+
+	// we now have our special case: view in tracker space history
+	// Pick the last one with matching frametime, as the 'latest' recorded query should have most accurate location.
+	// We might want to convert to UE FTransform later on to facilitate easy transformations.
+
+	TArray<FOpenXRLocateSpacePacket> SortedSpaceHistory = SpaceHistory;
+	SortedSpaceHistory.Sort([](const FOpenXRLocateSpacePacket& A, const FOpenXRLocateSpacePacket& B)
+	{
+		return A.Time < B.Time;
+	});
+
+	// We could also do this in place, putting elements in the front of the sorted array
+	// Basically 'remove duplicates'
+	// Alt: use iterators instead of raw indices
+	TArray<FOpenXRLocateSpacePacket> FilteredSpaceHistory;
+	for (int32 LocationIndex = 0; LocationIndex < SortedSpaceHistory.Num() - 1; LocationIndex++)
+	{
+		if (SortedSpaceHistory[LocationIndex].Time != SortedSpaceHistory[LocationIndex + 1].Time)
+		{
+			FilteredSpaceHistory.Add(SortedSpaceHistory[LocationIndex]);
+		}
+	}
+	FilteredSpaceHistory.Add(SortedSpaceHistory.Last());
+
+	// in the 'future' code, this is where we'd convert these poses to primary tracker-space
+	ReferencePoseHistories.Add(XR_REFERENCE_SPACE_TYPE_VIEW, MoveTemp(FilteredSpaceHistory));
+
+	// TODO: hypothetically, we can create multiple reference spaces from the same reference space type
+	// So we'd want to merge our histories
+}
+
+void FOpenXRPoseManager::RegisterEmulatedReferenceSpace(const XrReferenceSpaceCreateInfo& CreateInfo, XrSpace SpaceHandle)
+{
+	check(!EmulatedReferenceSpaceTypeMap.Contains(SpaceHandle));
+	EmulatedReferenceSpaceTypeMap.Add(SpaceHandle, CreateInfo.referenceSpaceType);
+}
+
+void FOpenXRPoseManager::AddEmulatedFrameTime(XrTime Time, int32 FrameNum)
+{
+	// TODO: protect against repeated frame times. would a runtime even do that? I guess it could happen,
+	// but this is for our runtime, and we won't do that for now
+	// TODO: protect against repeated frame numbers?
+
+	const int32 HistoryIndex = FrameNum % PoseHistorySize;
+	FrameTimeHistory[HistoryIndex] = FrameTimeHistoryEntry({ Time , FrameNum });
+
+	LastInsertedFrameIndex = HistoryIndex;
+}
+
+XrSpaceLocation FOpenXRPoseManager::GetEmulatedPoseForTime(XrSpace LocatingSpace, XrSpace BaseSpace, XrTime Time)
+{
+	// artifically constrain to view located in stage space
+	if (!EmulatedReferenceSpaceTypeMap.Contains(BaseSpace) ||
+		EmulatedReferenceSpaceTypeMap[BaseSpace] != XR_REFERENCE_SPACE_TYPE_STAGE ||
+		!EmulatedReferenceSpaceTypeMap.Contains(LocatingSpace) ||
+		EmulatedReferenceSpaceTypeMap[LocatingSpace] != XR_REFERENCE_SPACE_TYPE_VIEW)
+	{
+		XrSpaceLocation DummyLocation;
+		DummyLocation.locationFlags = XR_SPACE_LOCATION_ORIENTATION_VALID_BIT | XR_SPACE_LOCATION_POSITION_VALID_BIT | XR_SPACE_LOCATION_ORIENTATION_TRACKED_BIT | XR_SPACE_LOCATION_POSITION_TRACKED_BIT;
+		DummyLocation.pose.orientation = ToXrQuat(FQuat::Identity);
+		DummyLocation.pose.position = ToXrVector(FVector::ZeroVector);
+
+		return DummyLocation;
+	}
+
+	int32 SearchIndex = LastInsertedFrameIndex;
+	int32 StopIndex = SearchIndex - PoseHistorySize;
+	StopIndex = (StopIndex < 1) ? 1 : StopIndex;
+	while (SearchIndex >= StopIndex)
+	{
+		const int32 HistoryIndex = SearchIndex % PoseHistorySize;
+		if (FrameTimeHistory[HistoryIndex].Time == Time)
+		{
+			break;
+		}
+
+		SearchIndex--;
+	}
+
+	if (SearchIndex < StopIndex)
+	{
+		// We didn't find a match, just use most recent pose
+		SearchIndex = LastInsertedFrameIndex;
+	}
+
+	const int32 PoseIndex = FrameTimeHistory[SearchIndex].FrameIndex;
+
+	const TArray<FOpenXRLocateSpacePacket>& CapturedPoseHistories = ReferencePoseHistories[XR_REFERENCE_SPACE_TYPE_VIEW];
+	const int32 CountCapturedPoses = CapturedPoseHistories.Num();
+	const int32 CapturedIndex = PoseIndex % CountCapturedPoses;
+
+	return CapturedPoseHistories[CapturedIndex].Location;
+}
+
+} // namespace UE::XRScribe
diff --git a/Engine/Plugins/Runtime/XR/XRScribe/Source/XRScribe/Public/XRScribeAPIDecoder.h b/Engine/Plugins/Runtime/XR/XRScribe/Source/XRScribe/Public/XRScribeAPIDecoder.h
new file mode 100644
--- /dev/null
+++ b/Engine/Plugins/Runtime/XR/XRScribe/Source/XRScribe/Public/XRScribeAPIDecoder.h
@@ -0,0 +1,163 @@
+// Copyright Epic Games, Inc. All Rights Reserved.
+
+#pragma once
+
+#include "OpenXRPlatformRHI.h"
+#include "Containers/Map.h"
+#include "Serialization/ArrayReader.h"
+#include "XRScribeFileFormat.h"
+
+namespace UE::XRScribe
+{
+
+class FOpenXRCaptureDecoder
+{
+public:
+	explicit FOpenXRCaptureDecoder();
+	~FOpenXRCaptureDecoder();
+
+	void DecodeDataFromMemory();
+
+	// state accessors
+	[[nodiscard]] const TArray<XrExtensionProperties>& GetInstanceExtensionProperties() { return InstanceExtensionProperties; }
+	[[nodiscard]] const TArray<XrApiLayerProperties>& GetApiLayerProperties() {	return ApiLayerProperties; }
+	[[nodiscard]] XrInstanceCreateFlags GetInstanceCreateFlags() { return ValidInstanceCreateFlags; }
+	[[nodiscard]] const TArray<TStaticArray<ANSICHAR, XR_MAX_API_LAYER_NAME_SIZE>>& GetRequestedApiLayerNames() { return RequestedLayerNames; }
+	[[nodiscard]] const TArray<TStaticArray<ANSICHAR, XR_MAX_EXTENSION_NAME_SIZE>>& GetRequestedExtensionNames() { return RequestedExtensionNames; }
+	[[nodiscard]] const XrInstanceProperties& GetInstanceProperties() { return InstanceProperties; }
+	[[nodiscard]] const XrSystemGetInfo& GetSystemInfo() { return SystemGetInfo; }
+	[[nodiscard]] const XrSystemProperties& GetSystemProperties() { return SystemProperties; }
+	[[nodiscard]] const TArray<XrEnvironmentBlendMode>& GetEnvironmentBlendModes() { return EnvironmentBlendModes; }
+	[[nodiscard]] const TArray<XrViewConfigurationType>& GetViewConfigurationTypes() { return ViewConfigurationTypes; }
+	[[nodiscard]] const TMap<XrViewConfigurationType, XrViewConfigurationProperties>& GetViewConfigurationProperties() { return ViewConfigurationProperties; }
+	[[nodiscard]] const TMap<XrViewConfigurationType, TArray<XrViewConfigurationView>>& GetViewConfigurationViews() { return ViewConfigurationViews; }
+	[[nodiscard]] const TMap<XrViewConfigurationType, TArray<FOpenXRLocateViewsPacket>>& GetViewLocations() { return ViewLocations; }
+	[[nodiscard]] const TArray<FOpenXRCreateReferenceSpacePacket>& GetCreatedReferenceSpaces() { return CreatedReferenceSpaces; }
+	[[nodiscard]] const TMap<XrSpace, TArray<FOpenXRLocateSpacePacket>>& GetSpaceLocations() { return SpaceLocations; }
+	[[nodiscard]] const TArray<XrReferenceSpaceType>& GetReferenceSpaceTypes() { return ReferenceSpaceTypes; }
+	[[nodiscard]] const TMap<XrReferenceSpaceType, XrExtent2Df>& GetReferenceSpaceBounds() { return ReferenceSpaceBounds; }
+	[[nodiscard]] const TArray<int64>& GetSwapchainFormats() { return SwapchainFormats; }
+
+	TArray<uint8>& GetEncodedData()
+	{
+		return EncodedData;
+	}
+
+protected:
+
+	// packet decoders
+	void DecodeEnumerateApiLayerProperties(const FOpenXRAPIPacketBase& BasePacket);
+	void DecodeEnumerateInstanceExtensionProperties(const FOpenXRAPIPacketBase& BasePacket);
+	void DecodeCreateInstance(const FOpenXRAPIPacketBase& BasePacket);
+	void DecodeDestroyInstance(const FOpenXRAPIPacketBase& BasePacket);
+	void DecodeGetInstanceProperties(const FOpenXRAPIPacketBase& BasePacket);
+	//void DecodePollEvent(const FOpenXRAPIPacketBase& BasePacket);
+	//void DecodeResultToString(const FOpenXRAPIPacketBase& BasePacket);
+	//void DecodeStructureTypeToString(const FOpenXRAPIPacketBase& BasePacket);
+	void DecodeGetSystem(const FOpenXRAPIPacketBase& BasePacket);
+	void DecodeGetSystemProperties(const FOpenXRAPIPacketBase& BasePacket);
+	void DecodeEnumerateEnvironmentBlendModes(const FOpenXRAPIPacketBase& BasePacket);
+	void DecodeCreateSession(const FOpenXRAPIPacketBase& BasePacket);
+	void DecodeDestroySession(const FOpenXRAPIPacketBase& BasePacket);
+	void DecodeEnumerateReferenceSpaces(const FOpenXRAPIPacketBase& BasePacket);
+	void DecodeCreateReferenceSpace(const FOpenXRAPIPacketBase& BasePacket);
+	void DecodeGetReferenceSpaceBoundsRect(const FOpenXRAPIPacketBase& BasePacket);
+	void DecodeCreateActionSpace(const FOpenXRAPIPacketBase& BasePacket);
+	void DecodeLocateSpace(const FOpenXRAPIPacketBase& BasePacket);
+	void DecodeDestroySpace(const FOpenXRAPIPacketBase& BasePacket);
+	void DecodeEnumerateViewConfigurations(const FOpenXRAPIPacketBase& BasePacket);
+	void DecodeGetViewConfigurationProperties(const FOpenXRAPIPacketBase& BasePacket);
+	void DecodeEnumerateViewConfigurationViews(const FOpenXRAPIPacketBase& BasePacket);
+	void DecodeEnumerateSwapchainFormats(const FOpenXRAPIPacketBase& BasePacket);
+	void DecodeCreateSwapchain(const FOpenXRAPIPacketBase& BasePacket);
+	void DecodeDestroySwapchain(const FOpenXRAPIPacketBase& BasePacket);
+	void DecodeEnumerateSwapchainImages(const FOpenXRAPIPacketBase& BasePacket);
+	void DecodeAcquireSwapchainImage(const FOpenXRAPIPacketBase& BasePacket);
+	void DecodeWaitSwapchainImage(const FOpenXRAPIPacketBase& BasePacket);
+	void DecodeReleaseSwapchainImage(const FOpenXRAPIPacketBase& BasePacket);
+	void DecodeBeginSession(const FOpenXRAPIPacketBase& BasePacket);
+	void DecodeEndSession(const FOpenXRAPIPacketBase& BasePacket);
+	void DecodeRequestExitSession(const FOpenXRAPIPacketBase& BasePacket);
+	void DecodeWaitFrame(const FOpenXRAPIPacketBase& BasePacket);
+	void DecodeBeginFrame(const FOpenXRAPIPacketBase& BasePacket);
+	void DecodeEndFrame(const FOpenXRAPIPacketBase& BasePacket);
+	void DecodeLocateViews(const FOpenXRAPIPacketBase& BasePacket);
+	void DecodeStringToPath(const FOpenXRAPIPacketBase& BasePacket);
+	void DecodePathToString(const FOpenXRAPIPacketBase& BasePacket);
+	void DecodeCreateActionSet(const FOpenXRAPIPacketBase& BasePacket);
+	void DecodeDestroyActionSet(const FOpenXRAPIPacketBase& BasePacket);
+	void DecodeCreateAction(const FOpenXRAPIPacketBase& BasePacket);
+	void DecodeDestroyAction(const FOpenXRAPIPacketBase& BasePacket);
+
+	void DecodeSuggestInteractionProfileBindings(const FOpenXRAPIPacketBase& BasePacket);
+	void DecodeAttachSessionActionSets(const FOpenXRAPIPacketBase& BasePacket);
+	void DecodeGetCurrentInteractionProfile(const FOpenXRAPIPacketBase& BasePacket);
+	void DecodeGetActionStateBoolean(const FOpenXRAPIPacketBase& BasePacket);
+	void DecodeGetActionStateFloat(const FOpenXRAPIPacketBase& BasePacket);
+	void DecodeGetActionStateVector2f(const FOpenXRAPIPacketBase& BasePacket);
+	void DecodeGetActionStatePose(const FOpenXRAPIPacketBase& BasePacket);
+	void DecodeSyncActions(const FOpenXRAPIPacketBase& BasePacket);
+	//void DecodeEnumerateBoundSourcesForAction(const FOpenXRAPIPacketBase& BasePacket);
+	//void DecodeGetInputSourceLocalizedName(const FOpenXRAPIPacketBase& BasePacket);
+
+	void DecodeApplyHapticFeedback(const FOpenXRAPIPacketBase& BasePacket);
+	void DecodeStopHapticFeedback(const FOpenXRAPIPacketBase& BasePacket);
+	void DecodeInitializeLoaderKHR(const FOpenXRAPIPacketBase& BasePacket);
+	void DecodeGetVisibilityMaskKHR(const FOpenXRAPIPacketBase& BasePacket);
+#if defined(XR_USE_GRAPHICS_API_D3D11)
+	void DecodeGetD3D11GraphicsRequirementsKHR(const FOpenXRAPIPacketBase& BasePacket);
+#endif
+#if defined(XR_USE_GRAPHICS_API_D3D12)
+	void DecodeGetD3D12GraphicsRequirementsKHR(const FOpenXRAPIPacketBase& BasePacket);
+#endif
+
+	FArrayReader EncodedData;
+
+	typedef void(FOpenXRCaptureDecoder::* ApiDecodeFn)(const FOpenXRAPIPacketBase& BasePacket);
+	TStaticArray<ApiDecodeFn, (uint32)EOpenXRAPIPacketId::NumValidAPIPacketIds> DecodeFnTable;
+
+	// derived state from capture
+
+	TArray<XrExtensionProperties> InstanceExtensionProperties;
+	TArray<XrApiLayerProperties> ApiLayerProperties;
+	// TODO: Per-layer extension properties
+
+	XrInstanceCreateFlags ValidInstanceCreateFlags = 0;
+	TArray<TStaticArray<ANSICHAR, XR_MAX_API_LAYER_NAME_SIZE>> RequestedLayerNames;
+	TArray<TStaticArray<ANSICHAR, XR_MAX_EXTENSION_NAME_SIZE>> RequestedExtensionNames;
+
+	XrInstanceProperties InstanceProperties = {};
+
+	XrSystemGetInfo SystemGetInfo = {};
+	XrSystemProperties SystemProperties = {};
+	TArray<XrEnvironmentBlendMode> EnvironmentBlendModes;
+
+	XrSessionCreateInfo SessionCreateInfo = {};
+
+	TArray<XrReferenceSpaceType> ReferenceSpaceTypes;
+	TMap<XrReferenceSpaceType, XrExtent2Df> ReferenceSpaceBounds;
+	TMap<XrSpace, XrReferenceSpaceType> ReferenceSpaceMap;
+	TArray<FOpenXRCreateReferenceSpacePacket> CreatedReferenceSpaces;
+
+	TMap<XrSpace, XrAction> ActionSpaceMap;
+
+	TMap<XrSpace, TArray<FOpenXRLocateSpacePacket>> SpaceLocations;
+
+	TArray<XrViewConfigurationType> ViewConfigurationTypes;
+	TMap<XrViewConfigurationType, XrViewConfigurationProperties> ViewConfigurationProperties;
+	TMap<XrViewConfigurationType, TArray<XrViewConfigurationView>> ViewConfigurationViews;
+
+	TArray<int64> SwapchainFormats;
+
+	TMap<XrViewConfigurationType, TArray<FOpenXRLocateViewsPacket>> ViewLocations;
+
+	TMap<XrPath, FString> PathToStringMap;
+	TMap<FString, TArray<XrActionSuggestedBinding>> StringToSuggestedBindingsMap;
+
+	// TODO: Would I ever want to bin properties into per-instance collections?
+	// When we are repaying, we're just going to create our own set of 'valid' parameters
+	// that our replay runtime supports, do we really care about how the original application run
+	// managed their instances?
+};
+
+}
diff --git a/Engine/Plugins/Runtime/XR/XRScribe/Source/XRScribe/Public/XRScribeDeveloperSettings.h b/Engine/Plugins/Runtime/XR/XRScribe/Source/XRScribe/Public/XRScribeDeveloperSettings.h
new file mode 100644
--- /dev/null
+++ b/Engine/Plugins/Runtime/XR/XRScribe/Source/XRScribe/Public/XRScribeDeveloperSettings.h
@@ -0,0 +1,30 @@
+// Copyright Epic Games, Inc. All Rights Reserved.
+
+#pragma once
+
+// TODO: Use UDeveloperSettings + Project Settings UI in editor
+// Ideally, we'd use CDO UXRScribeDeveloperSettings to fetch our dev settings
+// however, we have some issues now with plugin loading + uobject setup
+// Even if we 'fix' those issues, reading the config into the class isn't stood up this early
+
+//#include "Engine/DeveloperSettings.h"
+//#include "XRScribeDeveloperSettings.generated.h"
+//
+///** Developer settings for XRScribe */
+//UCLASS(config = Engine, defaultconfig, meta = (DisplayName = "XRScribe"))
+//class XRSCRIBE_API UXRScribeDeveloperSettings : public UDeveloperSettings
+//{
+//	GENERATED_BODY()
+//public:
+//
+//	UXRScribeDeveloperSettings(const FObjectInitializer& Initializer);
+//
+//	/** Controls whether XRScribe runs in capture or emulation mode. Currently needed at engine startup, but will be runtime switchable. */
+//	UPROPERTY(config, EditAnywhere, Category = "XRScribe")
+//	EXRScribeRunMode RunMode = EXRScribeRunMode::Capture;
+//
+//	// TODO:
+//	// File path for capture file
+//	// Customizing capture dump point (session end, instance teardown, app end)
+//	// other run modes (e.g. replay)
+//};
\ No newline at end of file
diff --git a/Engine/Plugins/Runtime/XR/XRScribe/Source/XRScribe/XRScribe.Build.cs b/Engine/Plugins/Runtime/XR/XRScribe/Source/XRScribe/XRScribe.Build.cs
new file mode 100644
--- /dev/null
+++ b/Engine/Plugins/Runtime/XR/XRScribe/Source/XRScribe/XRScribe.Build.cs
@@ -0,0 +1,73 @@
+// Copyright Epic Games, Inc. All Rights Reserved.
+
+using System;
+using System.IO;
+
+namespace UnrealBuildTool.Rules
+{
+	// TODO Possibly split into capture module and replay module??
+	public class XRScribe : ModuleRules
+	{
+        public XRScribe(ReadOnlyTargetRules Target) : base(Target)
+		{
+			PublicDependencyModuleNames.AddRange(new string[]
+			{
+				"Core",
+				"CoreUObject",
+				"Engine",
+				"OpenXR",
+				"OpenXRHMD",
+				"RHI"
+			});
+
+			PrivateDependencyModuleNames.AddRange(new string[]
+			{
+                "DeveloperSettings",
+			});
+
+			// To get access to OpenXRPlatformRHI
+			PublicIncludePaths.Add(System.IO.Path.Combine(GetModuleDirectory("OpenXRHMD"), "Private"));
+
+			// For the RHI/SDK modules, we'll make both public so XRScribeTests can access them.
+			// We could just make the SDKs private, and have the tests include the headers. But we
+			// might actually want to do something in the tests eventually.
+
+			if (Target.Platform == UnrealTargetPlatform.Win64)
+			{
+				PublicDependencyModuleNames.AddRange(new string[] {
+					"D3D11RHI",
+					"D3D12RHI"
+				});
+
+				if (!bUsePrecompiled || Target.LinkType == TargetLinkType.Monolithic)
+				{
+					PublicDependencyModuleNames.AddRange(new string[] {
+						"DX11",
+						"DX12"
+					});
+				}
+			}
+
+			if (Target.Platform == UnrealTargetPlatform.Win64 || Target.Platform == UnrealTargetPlatform.Android)
+			{
+				PublicDependencyModuleNames.Add("OpenGLDrv");
+
+				if (!bUsePrecompiled || Target.LinkType == TargetLinkType.Monolithic)
+				{
+					PublicDependencyModuleNames.Add("OpenGL");
+				}
+			}
+
+			if (Target.Platform == UnrealTargetPlatform.Win64 || Target.Platform == UnrealTargetPlatform.Android
+				|| Target.IsInPlatformGroup(UnrealPlatformGroup.Linux))
+			{
+				PublicDependencyModuleNames.Add("VulkanRHI");
+
+				if (!bUsePrecompiled || Target.LinkType == TargetLinkType.Monolithic)
+				{
+					PublicDependencyModuleNames.Add("Vulkan");
+				}
+			}
+		}
+	}
+}
