XRScribe: fix emulation tests

D3D12 binding was scoped incorrectly in test

diff --git a/Engine/Plugins/Runtime/XR/XRScribe/Source/XRScribe/Private/XRScribeEmulationLayer.cpp b/Engine/Plugins/Runtime/XR/XRScribe/Source/XRScribe/Private/XRScribeEmulationLayer.cpp
--- a/Engine/Plugins/Runtime/XR/XRScribe/Source/XRScribe/Private/XRScribeEmulationLayer.cpp
+++ b/Engine/Plugins/Runtime/XR/XRScribe/Source/XRScribe/Private/XRScribeEmulationLayer.cpp
@@ -231,6 +231,10 @@ FOpenXREmulationLayer::FOpenXREmulationLayer()
 {
 	// We have to offset the time negatively because the first calls could be fast and fetch the same time!
 	StartTimeTicks = FDateTime::Now().GetTicks() - 1e7;
+
+#if defined(XR_USE_GRAPHICS_API_D3D12)
+	UE_LOG(LogXRScribeEmulate, Log, TEXT("XRScribe emulation supports D3D12"));
+#endif
 }
 FOpenXREmulationLayer::~FOpenXREmulationLayer() = default;
 
