XRScribe: Action space pose tracking

* share code between reference space and action space histories
* Build history based on pose input action names
* Register emulated actions and action spaces in order to match action name to capture
* Signal to emulation layer that an action is 'active' in xrGetActionStatePose

diff --git a/Engine/Plugins/Runtime/XR/XRScribe/Source/XRScribe/Private/XRScribeAPIDecoder.cpp b/Engine/Plugins/Runtime/XR/XRScribe/Source/XRScribe/Private/XRScribeAPIDecoder.cpp
--- a/Engine/Plugins/Runtime/XR/XRScribe/Source/XRScribe/Private/XRScribeAPIDecoder.cpp
+++ b/Engine/Plugins/Runtime/XR/XRScribe/Source/XRScribe/Private/XRScribeAPIDecoder.cpp
@@ -308,6 +308,7 @@ void FOpenXRCaptureDecoder::DecodeCreateActionSpace(const FOpenXRAPIPacketBase&
 	EncodedData << Data;
 
 	ActionSpaceMap.Add(Data.Space, Data.ActionSpaceCreateInfo.action);
+	CreatedActionSpaces.Add(Data);
 }
 
 void FOpenXRCaptureDecoder::DecodeLocateSpace(const FOpenXRAPIPacketBase& BasePacket)
@@ -613,7 +614,8 @@ void FOpenXRCaptureDecoder::DecodeCreateAction(const FOpenXRAPIPacketBase& BaseP
 	FOpenXRCreateActionPacket Data(XrResult::XR_ERROR_RUNTIME_FAILURE);
 
 	EncodedData << Data;
-	// nothing to do yet
+
+	CreatedActions.Add(Data);
 }
 
 void FOpenXRCaptureDecoder::DecodeDestroyAction(const FOpenXRAPIPacketBase& BasePacket)
diff --git a/Engine/Plugins/Runtime/XR/XRScribe/Source/XRScribe/Private/XRScribeEmulatedPoseManager.cpp b/Engine/Plugins/Runtime/XR/XRScribe/Source/XRScribe/Private/XRScribeEmulatedPoseManager.cpp
--- a/Engine/Plugins/Runtime/XR/XRScribe/Source/XRScribe/Private/XRScribeEmulatedPoseManager.cpp
+++ b/Engine/Plugins/Runtime/XR/XRScribe/Source/XRScribe/Private/XRScribeEmulatedPoseManager.cpp
@@ -3,9 +3,20 @@
 #include "XRScribeEmulatedPoseManager.h"
 #include "XRScribeFileFormat.h"
 
+#include "Logging/LogMacros.h"
+
+DEFINE_LOG_CATEGORY_STATIC(LogXRScribePoseManager, Log, All);
+
 namespace UE::XRScribe
 {
 
+void FOpenXRPoseManager::RegisterCapturedReferenceSpaces(const TArray<FOpenXRCreateReferenceSpacePacket>& CreateReferenceSpacePackets)
+{
+	CapturedReferenceSpaces.Append(CreateReferenceSpacePackets);
+
+	// TODO: check for doubles?
+}
+
 XrReferenceSpaceType FOpenXRPoseManager::GetCapturedReferenceSpaceType(XrSpace CapturedSpace)
 {
 	XrReferenceSpaceType Result = XR_REFERENCE_SPACE_TYPE_MAX_ENUM;
@@ -23,77 +34,142 @@ XrReferenceSpaceType FOpenXRPoseManager::GetCapturedReferenceSpaceType(XrSpace C
 	return Result;
 }
 
-void FOpenXRPoseManager::RegisterCapturedReferenceSpace(const FOpenXRCreateReferenceSpacePacket& CreateReferenceSpacePacket)
+void FOpenXRPoseManager::RegisterCapturedActions(const TArray <FOpenXRCreateActionPacket>& CreateActionPackets)
 {
-	CapturedReferenceSpaces.Add(CreateReferenceSpacePacket);
+	for (const FOpenXRCreateActionPacket& CreateAction : CreateActionPackets)
+	{
+		if (CreateAction.ActionType == XR_ACTION_TYPE_POSE_INPUT)
+		{
+			CapturedPoseActions.Add(CreateAction.Action, CreateAction);
+		}
+	}
+}
 
-	// TODO: check for doubles?
+void FOpenXRPoseManager::RegisterCapturedActionSpaces(const TArray<FOpenXRCreateActionSpacePacket>& CreateActionSpacePackets)
+{
+	CapturedActionSpaces.Append(CreateActionSpacePackets);
+}
 
-	if ((CreateReferenceSpacePacket.Result == XR_SUCCESS) &&
-		(CreateReferenceSpacePacket.ReferenceSpaceCreateInfo.referenceSpaceType == XR_REFERENCE_SPACE_TYPE_STAGE) &&
-		TrackingSpaceType != XR_REFERENCE_SPACE_TYPE_STAGE)
+bool FOpenXRPoseManager::VerifyCapturedActionSpace(XrSpace CapturedSpace)
+{
+	return CapturedActionSpaces.ContainsByPredicate([CapturedSpace](const FOpenXRCreateActionSpacePacket& CreateActionSpaceInfo)
 	{
-		TrackingSpaceType = XR_REFERENCE_SPACE_TYPE_STAGE;
-	}
+		return CreateActionSpaceInfo.Space == CapturedSpace;
+	});
 }
 
-void FOpenXRPoseManager::RegisterCapturedSpaceHistory(const TArray<FOpenXRLocateSpacePacket>& SpaceHistory)
+bool FOpenXRPoseManager::FilterSpaceHistory(const TArray<FOpenXRLocateSpacePacket>& RawHistory, TArray<FOpenXRLocateSpacePacket>& FilteredHistory)
 {
-	if (SpaceHistory.Num() == 0)
+	XrSpace CapturedBaseSpace = RawHistory[0].BaseSpace;
+	XrReferenceSpaceType BaseRefSpaceType = GetCapturedReferenceSpaceType(CapturedBaseSpace);
+
+	if (BaseRefSpaceType != XR_REFERENCE_SPACE_TYPE_STAGE)
 	{
-		// error info
-		return;
+		UE_LOG(LogXRScribePoseManager, Warning, TEXT("Unable to process space history with non-stage base space"));
+		return false;
 	}
 
-	XrSpace CapturedLocatedSpace = SpaceHistory[0].Space;
-	XrReferenceSpaceType RefSpaceType = GetCapturedReferenceSpaceType(CapturedLocatedSpace);
+	// we now have our core case: locate space relative to stage tracker
+	// Pick the last one with matching frametime, as the 'latest' recorded query should have most accurate location.
+	// We might want to convert to UE FTransform later on to facilitate easy transformations.
+	// TODO: account for different starting poses in stage space
 
-	if (RefSpaceType == XR_REFERENCE_SPACE_TYPE_MAX_ENUM)
+	FilteredHistory = RawHistory;
+	FilteredHistory.Sort([](const FOpenXRLocateSpacePacket& A, const FOpenXRLocateSpacePacket& B)
 	{
-		return;
+		return A.Time < B.Time;
+	});
+
+	// We could also do this in place like Algo::Unique, putting elements in the front of the sorted array
+	// But we want to select the 'last' element in the run sequence
+	int32 InsertIndex = 0;
+	for (int32 LocationIndex = 0; LocationIndex < FilteredHistory.Num() - 1; LocationIndex++)
+	{
+		if (FilteredHistory[LocationIndex].Time != FilteredHistory[LocationIndex + 1].Time)
+		{
+			FilteredHistory[InsertIndex++] = FilteredHistory[LocationIndex];
+		}
 	}
+	FilteredHistory[InsertIndex++] = FilteredHistory.Last();
+
+	FilteredHistory.SetNumUninitialized(InsertIndex);
+
+	// TODO: clients should transform these sorted poses into a unified space independent of changing base space
+
+	return true;
+}
+
+void FOpenXRPoseManager::ProcessCapturedReferenceSpaceHistory(const TArray<FOpenXRLocateSpacePacket>& SpaceHistory)
+{
+	const XrSpace CapturedLocatedSpace = SpaceHistory[0].Space;
+	const XrReferenceSpaceType RefSpaceType = GetCapturedReferenceSpaceType(CapturedLocatedSpace);
 
 	if (RefSpaceType != XR_REFERENCE_SPACE_TYPE_VIEW)
 	{
+		UE_LOG(LogXRScribePoseManager, Warning, TEXT("Unable to process space history for non-view reference spaces"));
 		return;
 	}
 
-	XrSpace CapturedBaseSpace = SpaceHistory[0].BaseSpace;
-	XrReferenceSpaceType BaseRefSpaceType = GetCapturedReferenceSpaceType(CapturedBaseSpace);
-
-	if (BaseRefSpaceType != XR_REFERENCE_SPACE_TYPE_STAGE)
+	TArray<FOpenXRLocateSpacePacket> FilteredSpaceHistory;
+	if (FilterSpaceHistory(SpaceHistory, FilteredSpaceHistory))
 	{
-		return;
+		ReferencePoseHistories.Add(XR_REFERENCE_SPACE_TYPE_VIEW, MoveTemp(FilteredSpaceHistory));
 	}
+}
 
-	// we now have our special case: view in tracker space history
-	// Pick the last one with matching frametime, as the 'latest' recorded query should have most accurate location.
-	// We might want to convert to UE FTransform later on to facilitate easy transformations.
+void FOpenXRPoseManager::ProcessCapturedActionSpaceHistory(const TArray<FOpenXRLocateSpacePacket>& SpaceHistory)
+{
+	const XrSpace LocatedSpace = SpaceHistory[0].Space;
 
-	TArray<FOpenXRLocateSpacePacket> SortedSpaceHistory = SpaceHistory;
-	SortedSpaceHistory.Sort([](const FOpenXRLocateSpacePacket& A, const FOpenXRLocateSpacePacket& B)
+	FOpenXRCreateActionSpacePacket* MatchingInfo = CapturedActionSpaces.FindByPredicate([LocatedSpace](const FOpenXRCreateActionSpacePacket& CreateActionSpaceInfo)
 	{
-		return A.Time < B.Time;
+		return CreateActionSpaceInfo.Space == LocatedSpace;
+
 	});
+	check(MatchingInfo);
+
+	const FOpenXRCreateActionPacket& CreateAction = CapturedPoseActions[MatchingInfo->ActionSpaceCreateInfo.action];
+
+	check(CreateAction.ActionType == XR_ACTION_TYPE_POSE_INPUT);
 
-	// We could also do this in place, putting elements in the front of the sorted array
-	// Basically 'remove duplicates'
-	// Alt: use iterators instead of raw indices
 	TArray<FOpenXRLocateSpacePacket> FilteredSpaceHistory;
-	for (int32 LocationIndex = 0; LocationIndex < SortedSpaceHistory.Num() - 1; LocationIndex++)
+	if (FilterSpaceHistory(SpaceHistory, FilteredSpaceHistory))
 	{
-		if (SortedSpaceHistory[LocationIndex].Time != SortedSpaceHistory[LocationIndex + 1].Time)
+		const FName FetchedActionName(CreateAction.ActionName.GetData());
+		ActionPoseHistories.Add(FetchedActionName, MoveTemp(FilteredSpaceHistory));
+
+		// TODO: we should be resilient to multiple action spaces per name
+	}
+}
+
+void FOpenXRPoseManager::RegisterCapturedSpaceHistories(const TMap<XrSpace, TArray<FOpenXRLocateSpacePacket>>& SpaceHistories)
+{
+	for (const auto& SpaceHistoryTuple : SpaceHistories)
+	{
+		const TArray<FOpenXRLocateSpacePacket>& SpaceHistory = SpaceHistoryTuple.Value;
+		if (SpaceHistory.Num() == 0)
 		{
-			FilteredSpaceHistory.Add(SortedSpaceHistory[LocationIndex]);
+			continue;
 		}
-	}
-	FilteredSpaceHistory.Add(SortedSpaceHistory.Last());
 
-	// in the 'future' code, this is where we'd convert these poses to primary tracker-space
-	ReferencePoseHistories.Add(XR_REFERENCE_SPACE_TYPE_VIEW, MoveTemp(FilteredSpaceHistory));
+		const XrSpace CapturedLocatedSpace = SpaceHistoryTuple.Key;
+		check(CapturedLocatedSpace == SpaceHistory[0].Space);
+
+		const XrReferenceSpaceType RefSpaceType = GetCapturedReferenceSpaceType(CapturedLocatedSpace);
+		if (RefSpaceType != XR_REFERENCE_SPACE_TYPE_MAX_ENUM)
+		{
+			ProcessCapturedReferenceSpaceHistory(SpaceHistory);
+		}
+		else if (VerifyCapturedActionSpace(CapturedLocatedSpace))
+		{
+			ProcessCapturedActionSpaceHistory(SpaceHistory);
+		}
+		else
+		{
+			UE_LOG(LogXRScribePoseManager, Warning, TEXT("Unable to process space history for captured space: %p"), CapturedLocatedSpace);
+		}
 
-	// TODO: hypothetically, we can create multiple reference spaces from the same reference space type
-	// So we'd want to merge our histories
+	}
 }
 
 void FOpenXRPoseManager::RegisterEmulatedReferenceSpace(const XrReferenceSpaceCreateInfo& CreateInfo, XrSpace SpaceHandle)
@@ -102,6 +178,14 @@ void FOpenXRPoseManager::RegisterEmulatedReferenceSpace(const XrReferenceSpaceCr
 	EmulatedReferenceSpaceTypeMap.Add(SpaceHandle, CreateInfo.referenceSpaceType);
 }
 
+void FOpenXRPoseManager::RegisterEmulatedActionSpace(TStaticArray<ANSICHAR, XR_MAX_ACTION_NAME_SIZE>& ActionName, const XrActionSpaceCreateInfo& SpaceCreateInfo, XrSpace SpaceHandle)
+{
+	check(!EmulatedActionSpaceNameMap.Contains(SpaceHandle));
+
+	const FName EmulatedActionName(ActionName.GetData());
+	EmulatedActionSpaceNameMap.Add(SpaceHandle, EmulatedActionName);
+}
+
 void FOpenXRPoseManager::AddEmulatedFrameTime(XrTime Time, int32 FrameNum)
 {
 	// TODO: protect against repeated frame times. would a runtime even do that? I guess it could happen,
@@ -114,19 +198,22 @@ void FOpenXRPoseManager::AddEmulatedFrameTime(XrTime Time, int32 FrameNum)
 	LastInsertedFrameIndex = HistoryIndex;
 }
 
+bool FOpenXRPoseManager::DoesActionContainPoseHistory(TStaticArray<ANSICHAR, XR_MAX_ACTION_NAME_SIZE>& ActionName)
+{
+	return ActionPoseHistories.Contains(ActionName.GetData());
+}
+
 XrSpaceLocation FOpenXRPoseManager::GetEmulatedPoseForTime(XrSpace LocatingSpace, XrSpace BaseSpace, XrTime Time)
 {
-	// artifically constrain to view located in stage space
+	XrSpaceLocation DummyLocation;
+	DummyLocation.locationFlags = XR_SPACE_LOCATION_ORIENTATION_VALID_BIT | XR_SPACE_LOCATION_POSITION_VALID_BIT | XR_SPACE_LOCATION_ORIENTATION_TRACKED_BIT | XR_SPACE_LOCATION_POSITION_TRACKED_BIT;
+	DummyLocation.pose.orientation = ToXrQuat(FQuat::Identity);
+	DummyLocation.pose.position = ToXrVector(FVector::ZeroVector);
+
+	// artifically constrain locations to stage space
 	if (!EmulatedReferenceSpaceTypeMap.Contains(BaseSpace) ||
-		EmulatedReferenceSpaceTypeMap[BaseSpace] != XR_REFERENCE_SPACE_TYPE_STAGE ||
-		!EmulatedReferenceSpaceTypeMap.Contains(LocatingSpace) ||
-		EmulatedReferenceSpaceTypeMap[LocatingSpace] != XR_REFERENCE_SPACE_TYPE_VIEW)
+			EmulatedReferenceSpaceTypeMap[BaseSpace] != XR_REFERENCE_SPACE_TYPE_STAGE)
 	{
-		XrSpaceLocation DummyLocation;
-		DummyLocation.locationFlags = XR_SPACE_LOCATION_ORIENTATION_VALID_BIT | XR_SPACE_LOCATION_POSITION_VALID_BIT | XR_SPACE_LOCATION_ORIENTATION_TRACKED_BIT | XR_SPACE_LOCATION_POSITION_TRACKED_BIT;
-		DummyLocation.pose.orientation = ToXrQuat(FQuat::Identity);
-		DummyLocation.pose.position = ToXrVector(FVector::ZeroVector);
-
 		return DummyLocation;
 	}
 
@@ -152,11 +239,25 @@ XrSpaceLocation FOpenXRPoseManager::GetEmulatedPoseForTime(XrSpace LocatingSpace
 
 	const int32 PoseIndex = FrameTimeHistory[SearchIndex].FrameIndex;
 
-	const TArray<FOpenXRLocateSpacePacket>& CapturedPoseHistories = ReferencePoseHistories[XR_REFERENCE_SPACE_TYPE_VIEW];
-	const int32 CountCapturedPoses = CapturedPoseHistories.Num();
-	const int32 CapturedIndex = PoseIndex % CountCapturedPoses;
+	if (EmulatedReferenceSpaceTypeMap.Contains(LocatingSpace) && EmulatedReferenceSpaceTypeMap[LocatingSpace] == XR_REFERENCE_SPACE_TYPE_VIEW)
+	{
+		const TArray<FOpenXRLocateSpacePacket>& CapturedPoseHistories = ReferencePoseHistories[XR_REFERENCE_SPACE_TYPE_VIEW];
+		const int32 CountCapturedPoses = CapturedPoseHistories.Num();
+		const int32 CapturedIndex = PoseIndex % CountCapturedPoses;
+
+		return CapturedPoseHistories[CapturedIndex].Location;
+	}
+	else if (EmulatedActionSpaceNameMap.Contains(LocatingSpace) && ActionPoseHistories.Contains(EmulatedActionSpaceNameMap[LocatingSpace]))
+	{
+		const FName ActionName = EmulatedActionSpaceNameMap[LocatingSpace];
+		const TArray<FOpenXRLocateSpacePacket>& CapturedPoseHistories = ActionPoseHistories[ActionName]; // TODO: should actually check that that we have a history...
+		const int32 CountCapturedPoses = CapturedPoseHistories.Num();
+		const int32 CapturedIndex = PoseIndex % CountCapturedPoses;
+
+		return CapturedPoseHistories[CapturedIndex].Location;
+	}
 
-	return CapturedPoseHistories[CapturedIndex].Location;
+	return DummyLocation;
 }
 
 } // namespace UE::XRScribe
diff --git a/Engine/Plugins/Runtime/XR/XRScribe/Source/XRScribe/Private/XRScribeEmulationLayer.cpp b/Engine/Plugins/Runtime/XR/XRScribe/Source/XRScribe/Private/XRScribeEmulationLayer.cpp
--- a/Engine/Plugins/Runtime/XR/XRScribe/Source/XRScribe/Private/XRScribeEmulationLayer.cpp
+++ b/Engine/Plugins/Runtime/XR/XRScribe/Source/XRScribe/Private/XRScribeEmulationLayer.cpp
@@ -344,20 +344,12 @@ void FOpenXREmulationLayer::PostLoadActions()
 		FCStringAnsi::Strncpy(&EmulatedSystemProperties.systemName[ActualSystemNameLen], EmuStringAddend, RemainingSpace);
 	}
 
-	// TODO: build pose history relative to primary used tracking space (STAGE + LOCAL)
-	// TODO: build pose lookups relative to any reference space
+	PoseManager.RegisterCapturedReferenceSpaces(CaptureDecoder.GetCreatedReferenceSpaces());
 
-	const TArray<FOpenXRCreateReferenceSpacePacket>& CapturedRefSpaces = CaptureDecoder.GetCreatedReferenceSpaces();
-	for (const FOpenXRCreateReferenceSpacePacket& RefSpace : CapturedRefSpaces)
-	{
-		PoseManager.RegisterCapturedReferenceSpace(RefSpace);
-	}
+	PoseManager.RegisterCapturedActions(CaptureDecoder.GetCreatedActions());
+	PoseManager.RegisterCapturedActionSpaces(CaptureDecoder.GetCreatedActionSpaces());
 
-	const TMap<XrSpace, TArray<FOpenXRLocateSpacePacket>>& CapturedSpaceLocations = CaptureDecoder.GetSpaceLocations();
-	for (const auto& SpaceHistory : CapturedSpaceLocations)
-	{
-		PoseManager.RegisterCapturedSpaceHistory(SpaceHistory.Value);
-	}
+	PoseManager.RegisterCapturedSpaceHistories(CaptureDecoder.GetSpaceLocations());
 }
 
 bool FOpenXREmulationLayer::InstanceHandleCheck(XrInstance InputInstance)
@@ -1169,6 +1161,11 @@ XrResult FOpenXREmulationLayer::XrLayerCreateActionSpace(XrSession session, cons
 		*space = reinterpret_cast<XrSpace>(CurrentSession->ActiveSpaces.Last().Get());
 	}
 
+	{
+		FOpenXREmulatedAction* OwningAction = reinterpret_cast<FOpenXREmulatedAction*>(createInfo->action);
+		PoseManager.RegisterEmulatedActionSpace(OwningAction->ActionName, *createInfo, *space);
+	}
+
 	return XrResult::XR_SUCCESS;
 }
 
@@ -2288,6 +2285,13 @@ XrResult FOpenXREmulationLayer::XrLayerAttachSessionActionSets(XrSession session
 
 	// TODO: save off list of action sets + sets in session
 
+	for (uint32 ActionSetIndex = 0; ActionSetIndex < attachInfo->countActionSets; ActionSetIndex++)
+	{
+		FOpenXREmulatedActionSet* EmulatedActionSet = reinterpret_cast<FOpenXREmulatedActionSet*>(attachInfo->actionSets[ActionSetIndex]);
+		check(!EmulatedActionSet->bAttached);
+		EmulatedActionSet->bAttached = true;
+	}
+
 	{
 		FWriteScopeLock SessionLock(SessionMutex);
 		CurrentSession->bActionSetsAttached = true;
@@ -2432,11 +2436,26 @@ XrResult FOpenXREmulationLayer::XrLayerGetActionStatePose(XrSession session, con
 		return XR_ERROR_VALIDATION_FAILURE;
 	}
 
-	// TODO: check action exists and is attached
-	// TODO: check subaction path
+	if (!ActionHandleCheck(getInfo->action))
+	{
+		UE_LOG(LogXRScribeEmulate, Error, TEXT("GetActionStatePose: unknown action"));
+		return XR_ERROR_VALIDATION_FAILURE;
+	}
 
-	// TODO: fill XrActionStatePose with real data
-	state->isActive = XR_FALSE;
+	FOpenXREmulatedAction* Action = reinterpret_cast<FOpenXREmulatedAction*>(getInfo->action);
+	if (!Action->OwningActionSet->bAttached)
+	{
+		return XR_ERROR_ACTIONSET_NOT_ATTACHED;
+	}
+
+	if (Action->ActionType != XR_ACTION_TYPE_POSE_INPUT)
+	{
+		return XR_ERROR_ACTION_TYPE_MISMATCH;
+	}
+
+	// TODO: check subaction path?
+
+	state->isActive = PoseManager.DoesActionContainPoseHistory(Action->ActionName);
 
 	return XR_SUCCESS;
 }
diff --git a/Engine/Plugins/Runtime/XR/XRScribe/Source/XRScribe/Public/XRScribeAPIDecoder.h b/Engine/Plugins/Runtime/XR/XRScribe/Source/XRScribe/Public/XRScribeAPIDecoder.h
--- a/Engine/Plugins/Runtime/XR/XRScribe/Source/XRScribe/Public/XRScribeAPIDecoder.h
+++ b/Engine/Plugins/Runtime/XR/XRScribe/Source/XRScribe/Public/XRScribeAPIDecoder.h
@@ -33,10 +33,12 @@ public:
 	[[nodiscard]] const TMap<XrViewConfigurationType, TArray<XrViewConfigurationView>>& GetViewConfigurationViews() { return ViewConfigurationViews; }
 	[[nodiscard]] const TMap<XrViewConfigurationType, TArray<FOpenXRLocateViewsPacket>>& GetViewLocations() { return ViewLocations; }
 	[[nodiscard]] const TArray<FOpenXRCreateReferenceSpacePacket>& GetCreatedReferenceSpaces() { return CreatedReferenceSpaces; }
+	[[nodiscard]] const TArray<FOpenXRCreateActionSpacePacket>& GetCreatedActionSpaces() { return CreatedActionSpaces; }
 	[[nodiscard]] const TMap<XrSpace, TArray<FOpenXRLocateSpacePacket>>& GetSpaceLocations() { return SpaceLocations; }
 	[[nodiscard]] const TArray<XrReferenceSpaceType>& GetReferenceSpaceTypes() { return ReferenceSpaceTypes; }
 	[[nodiscard]] const TMap<XrReferenceSpaceType, XrExtent2Df>& GetReferenceSpaceBounds() { return ReferenceSpaceBounds; }
 	[[nodiscard]] const TArray<int64>& GetSwapchainFormats() { return SwapchainFormats; }
+	[[nodiscard]] const TArray<FOpenXRCreateActionPacket>& GetCreatedActions() { return CreatedActions; }
 
 	TArray<uint8>& GetEncodedData()
 	{
@@ -139,6 +141,7 @@ protected:
 	TMap<XrSpace, XrReferenceSpaceType> ReferenceSpaceMap;
 	TArray<FOpenXRCreateReferenceSpacePacket> CreatedReferenceSpaces;
 
+	TArray<FOpenXRCreateActionSpacePacket> CreatedActionSpaces;
 	TMap<XrSpace, XrAction> ActionSpaceMap;
 
 	TMap<XrSpace, TArray<FOpenXRLocateSpacePacket>> SpaceLocations;
@@ -154,6 +157,8 @@ protected:
 	TMap<XrPath, FString> PathToStringMap;
 	TMap<FString, TArray<XrActionSuggestedBinding>> StringToSuggestedBindingsMap;
 
+	TArray<FOpenXRCreateActionPacket> CreatedActions;
+
 	// TODO: Would I ever want to bin properties into per-instance collections?
 	// When we are repaying, we're just going to create our own set of 'valid' parameters
 	// that our replay runtime supports, do we really care about how the original application run
diff --git a/Engine/Plugins/Runtime/XR/XRScribe/Source/XRScribe/Public/XRScribeEmulatedPoseManager.h b/Engine/Plugins/Runtime/XR/XRScribe/Source/XRScribe/Public/XRScribeEmulatedPoseManager.h
--- a/Engine/Plugins/Runtime/XR/XRScribe/Source/XRScribe/Public/XRScribeEmulatedPoseManager.h
+++ b/Engine/Plugins/Runtime/XR/XRScribe/Source/XRScribe/Public/XRScribeEmulatedPoseManager.h
@@ -2,6 +2,8 @@
 
 #pragma once
 
+#include "XRScribeFileFormat.h"
+
 #include "Containers/Array.h"
 #include "Containers/Map.h"
 #include "Containers/StaticArray.h"
@@ -12,9 +14,6 @@
 namespace UE::XRScribe
 {
 
-struct FOpenXRCreateReferenceSpacePacket;
-struct FOpenXRLocateSpacePacket;
-
 // The PoseManager has two primary jobs:
 // * ingest data from capture
 // * generate estimated poses based on emulator frametimes
@@ -35,21 +34,39 @@ public:
 	static const int32 PoseHistorySize = 5;
 
 	/**
-	 * Register information about a reference space from a capture file, in order to build a history of
-	 * space locations relative to that reference space.
+	 * Register information about reference spaces from a capture file, in order to build a history of
+	 * space locations relative to those reference spaces.
 	 *
-	 * @param CreateReferenceSpacePacket The create info for the captured reference space
+	 * @param CreateReferenceSpacePackets The create infos for the captured reference spaces
 	 *
 	 */
-	void RegisterCapturedReferenceSpace(const FOpenXRCreateReferenceSpacePacket& CreateReferenceSpacePacket);
+	void RegisterCapturedReferenceSpaces(const TArray<FOpenXRCreateReferenceSpacePacket>& CreateReferenceSpacePackets);
 
 	/**
-	 * Register history of LocateSpace calls relative to a specific XrSpace
+	 * Register information about actions from capture file, in order to fetch action name for mapping between capture and emulation
 	 *
-	 * @param SpaceHistory List of LocateSpace calls with a unified Space (but possibly varying BaseSpace) 
+	 * @param CreateActionPackets The create infos for the captured actions
 	 *
 	 */
-	void RegisterCapturedSpaceHistory(const TArray<FOpenXRLocateSpacePacket>& SpaceHistory);
+	void RegisterCapturedActions(const TArray<FOpenXRCreateActionPacket>& CreateActionPackets);
+
+
+	/**
+	 * Register information about action spaces from capture file, which we need to build a history of space locations
+	 * for XR devices
+	 *
+	 * @param CreateActionSpacePackets The create infos for the captured action spaces
+	 *
+	 */
+	void RegisterCapturedActionSpaces(const TArray<FOpenXRCreateActionSpacePacket>& CreateActionSpacePackets);
+
+	/**
+	 * Register history of LocateSpace calls
+	 *
+	 * @param SpaceHistories List of LocateSpace histories, with a unified Space per history (but possibly varying BaseSpace) 
+	 *
+	 */
+	void RegisterCapturedSpaceHistories(const TMap<XrSpace, TArray<FOpenXRLocateSpacePacket>>& SpaceHistories);
 
 	/**
 	 * Register emulated reference space. Pose manager needs this to match the new emulated reference space against
@@ -61,6 +78,17 @@ public:
 	 */
 	void RegisterEmulatedReferenceSpace(const XrReferenceSpaceCreateInfo& CreateInfo, XrSpace SpaceHandle);
 
+	/**
+	 * Register emulated action space. Pose manager needs this to match emulated action space against
+	 * a captured action space with matching action name.
+	 *
+	 * @param ActionName Name from action create info
+	 * @param SpaceCreateInfo Create info for emulated action space, in order to match against captured history of similarly configured space
+	 * @param SpaceHandle Emulation generated handle for new space
+	 *
+	 */
+	void RegisterEmulatedActionSpace(TStaticArray<ANSICHAR, XR_MAX_ACTION_NAME_SIZE>& ActionName, const XrActionSpaceCreateInfo& SpaceCreateInfo, XrSpace SpaceHandle);
+
 	/**
 	 * Add time and frame number in order to generate interpolated poses
 	 * We use the frame numbers to index into a relative history of the poses, and the time
@@ -72,6 +100,15 @@ public:
 	 */
 	void AddEmulatedFrameTime(XrTime Time, int32 FrameNum);
 
+	/**
+	 * If we have a pose history for the action, then xrGetActionStatePose can
+	 * tell the client that the action is active, and poses are ready to fetch
+	 *
+	 * @param ActionName Name from action create info
+	 *
+	 */
+	bool DoesActionContainPoseHistory(TStaticArray<ANSICHAR, XR_MAX_ACTION_NAME_SIZE>& ActionName);
+
 	/**
 	 * Fetch a pose generated from the captured history, by using space handles and times generated by the emulated runtime
 	 *
@@ -83,27 +120,42 @@ public:
 	 */
 	XrSpaceLocation GetEmulatedPoseForTime(XrSpace LocatingSpace, XrSpace BaseSpace, XrTime Time);
 
+	// TODO: add LocateViews to pose manager tracking
+
 private:
 
-	// Known reference spaces from capture
-	TArray<FOpenXRCreateReferenceSpacePacket> CapturedReferenceSpaces;
-	
 	/**
-	 * Pose histories for reference space types, which could be unified across different space handles
-	 * with the same underlying reference space type.
+	 * Helper function to fetch underlying reference space type from opaque XrSpace handle
 	 */
-	TMap<XrReferenceSpaceType, TArray<FOpenXRLocateSpacePacket>> ReferencePoseHistories;
+	XrReferenceSpaceType GetCapturedReferenceSpaceType(XrSpace CapturedSpace);
 
-	XrReferenceSpaceType TrackingSpaceType = XR_REFERENCE_SPACE_TYPE_VIEW;
+	bool VerifyCapturedActionSpace(XrSpace CapturedSpace);
 
 	/**
-	 * Helper function to fetch underlying reference space type from opaque XrSpace handle
+	 * sort and remove duplicates in history for single space
 	 */
-	XrReferenceSpaceType GetCapturedReferenceSpaceType(XrSpace CapturedSpace);
+	bool FilterSpaceHistory(const TArray<FOpenXRLocateSpacePacket>& RawHistory, TArray<FOpenXRLocateSpacePacket>& FilteredHistory);
+
+	void ProcessCapturedReferenceSpaceHistory(const TArray<FOpenXRLocateSpacePacket>& SpaceHistory);
+	void ProcessCapturedActionSpaceHistory(const TArray<FOpenXRLocateSpacePacket>& SpaceHistory);
+
+	// Known spaces and relevant actions from capture
+	TArray<FOpenXRCreateReferenceSpacePacket> CapturedReferenceSpaces;
+	TArray<FOpenXRCreateActionSpacePacket> CapturedActionSpaces;
+	TMap<XrAction, FOpenXRCreateActionPacket> CapturedPoseActions;
+
+	/**
+	 * Pose histories for space types, which could be unified across different space handles
+	 * with the same underlying reference space type.
+	 */
+	TMap<XrReferenceSpaceType, TArray<FOpenXRLocateSpacePacket>> ReferencePoseHistories;
+	TMap<FName, TArray<FOpenXRLocateSpacePacket>> ActionPoseHistories;
 
 	// Map of emulated spaces to their underlying reference space type (if applicable)
 	TMap<XrSpace, XrReferenceSpaceType> EmulatedReferenceSpaceTypeMap;
 
+	TMap< XrSpace, FName> EmulatedActionSpaceNameMap;
+
 	struct FrameTimeHistoryEntry
 	{
 		XrTime Time = 0;
@@ -113,7 +165,6 @@ private:
 	// History of frame times, in order for pose manager to look into to estimate the relative frame index for times
 	TStaticArray <FrameTimeHistoryEntry, PoseHistorySize> FrameTimeHistory;
 	int32 LastInsertedFrameIndex = 0;
-
 };
 
 } // namespace UE::XRScribe
