XRScribe: handle layer setup failures

* Capture layer won't be set up without a valid PFN_xrGetInstanceProcAddr from the real loader
* Emulation layer won't set up without valid capture file

#rnx

diff --git a/Engine/Plugins/Runtime/XR/XRScribe/Source/XRScribe/Private/XRScribeAPISurface.cpp b/Engine/Plugins/Runtime/XR/XRScribe/Source/XRScribe/Private/XRScribeAPISurface.cpp
--- a/Engine/Plugins/Runtime/XR/XRScribe/Source/XRScribe/Private/XRScribeAPISurface.cpp
+++ b/Engine/Plugins/Runtime/XR/XRScribe/Source/XRScribe/Private/XRScribeAPISurface.cpp
@@ -697,10 +697,9 @@ public:
 		return ActiveLayer.Get();
 	}
 
-	void SetChainedGetProcAddr(PFN_xrGetInstanceProcAddr InChainedGetProcAddr) override
+	bool SetChainedGetProcAddr(PFN_xrGetInstanceProcAddr InChainedGetProcAddr) override
 	{
 		ChainedGetProcAddr = InChainedGetProcAddr;
-		CaptureLayer->SetChainedGetProcAddr(ChainedGetProcAddr);
 
 		EXRScribeRunMode RunModeFromConfig = EXRScribeRunMode::Capture;
 		GConfig->GetInt(TEXT("SystemSettings"), TEXT("XRScribe.RunMode"), (int32&)RunModeFromConfig, GEngineIni);
@@ -708,18 +707,42 @@ public:
 		if (RunModeFromConfig == EXRScribeRunMode::Capture)
 		{
 			UE_LOG(LogXRScribeAPI, Log, TEXT("Capture layer selected"));
-			ActiveLayer = CaptureLayer;
+
+			if (InChainedGetProcAddr != nullptr)
+			{
+				UE_LOG(LogXRScribeAPI, Log, TEXT("Valid PFN_xrGetInstanceProcAddr handed off to Capture Layer"));
+				ActiveLayer = CaptureLayer;
+				CaptureLayer->SetChainedGetProcAddr(ChainedGetProcAddr);
+				return true;
+			}
+			else
+			{
+				UE_LOG(LogXRScribeAPI, Warning, TEXT("Invalid PFN_xrGetInstanceProcAddr, Capture Layer disabled"));
+				return false;
+			}
 		}
 		else if (RunModeFromConfig == EXRScribeRunMode::Emulate)
 		{
 			UE_LOG(LogXRScribeAPI, Log, TEXT("Emulation layer selected"));
 			ActiveLayer = EmulationLayer;
-			reinterpret_cast<FOpenXREmulationLayer*>(EmulationLayer.Get())->LoadCaptureFromFile();
+			const bool bEmulationLoadedCapture = reinterpret_cast<FOpenXREmulationLayer*>(EmulationLayer.Get())->LoadCaptureFromFile();
+
+			if (bEmulationLoadedCapture)
+			{
+				UE_LOG(LogXRScribeAPI, Log, TEXT("Emulation Layer enabled"));
+			}
+			else
+			{
+				UE_LOG(LogXRScribeAPI, Warning, TEXT("Emulation Layer disabled"));
+			}
+
+			return bEmulationLoadedCapture;
 		}
 		else
 		{
 			check(0);
 		}
+		return false;
 	}
 
 private:
diff --git a/Engine/Plugins/Runtime/XR/XRScribe/Source/XRScribe/Private/XRScribeModule.cpp b/Engine/Plugins/Runtime/XR/XRScribe/Source/XRScribe/Private/XRScribeModule.cpp
--- a/Engine/Plugins/Runtime/XR/XRScribe/Source/XRScribe/Private/XRScribeModule.cpp
+++ b/Engine/Plugins/Runtime/XR/XRScribe/Source/XRScribe/Private/XRScribeModule.cpp
@@ -31,13 +31,16 @@ FString FXRScribeModule::GetDisplayName()
 
 bool FXRScribeModule::InsertOpenXRAPILayer(PFN_xrGetInstanceProcAddr& InOutGetProcAddr)
 {
-	ChainedGetProcAddr = InOutGetProcAddr;
-	InOutGetProcAddr = &UE::XRScribe::xrGetInstanceProcAddr;
-	
 	// Ok to pass nullptr for GetProcAddr because that would only disable capture. Replay still valid!
-	UE::XRScribe::IOpenXRAPILayerManager::Get().SetChainedGetProcAddr(ChainedGetProcAddr);
+	const bool bResult = UE::XRScribe::IOpenXRAPILayerManager::Get().SetChainedGetProcAddr(ChainedGetProcAddr);
 
-	return true;
+	if (bResult)
+	{
+		ChainedGetProcAddr = InOutGetProcAddr;
+		InOutGetProcAddr = &UE::XRScribe::xrGetInstanceProcAddr;
+	}
+
+	return bResult;
 }
 
 
diff --git a/Engine/Plugins/Runtime/XR/XRScribe/Source/XRScribe/Public/XRScribeAPISurface.h b/Engine/Plugins/Runtime/XR/XRScribe/Source/XRScribe/Public/XRScribeAPISurface.h
--- a/Engine/Plugins/Runtime/XR/XRScribe/Source/XRScribe/Public/XRScribeAPISurface.h
+++ b/Engine/Plugins/Runtime/XR/XRScribe/Source/XRScribe/Public/XRScribeAPISurface.h
@@ -22,7 +22,7 @@ class IOpenXRAPILayerManager
 public:
 	virtual IOpenXRAPILayer* GetActiveLayer() const = 0;
 
-	virtual void SetChainedGetProcAddr(PFN_xrGetInstanceProcAddr InChainedGetProcAddr) = 0;
+	virtual bool SetChainedGetProcAddr(PFN_xrGetInstanceProcAddr InChainedGetProcAddr) = 0;
 	
 	static IOpenXRAPILayerManager& Get();
 };
