XRScribe: emulate by default

We don't need to spuriously capture everytime we load. We should attempt to load a capture, and then move on with regular functionality if nothing to load

#rnx

diff --git a/Engine/Plugins/Runtime/XR/XRScribe/Source/XRScribe/Private/XRScribeAPISurface.cpp b/Engine/Plugins/Runtime/XR/XRScribe/Source/XRScribe/Private/XRScribeAPISurface.cpp
--- a/Engine/Plugins/Runtime/XR/XRScribe/Source/XRScribe/Private/XRScribeAPISurface.cpp
+++ b/Engine/Plugins/Runtime/XR/XRScribe/Source/XRScribe/Private/XRScribeAPISurface.cpp
@@ -32,7 +32,7 @@ enum class EXRScribeRunMode : int32
 	Emulate,
 };
 
-static EXRScribeRunMode XRScribeRunMode = EXRScribeRunMode::Capture;
+static EXRScribeRunMode XRScribeRunMode = EXRScribeRunMode::Emulate;
 static FAutoConsoleVariableRef CVarXRScribeRunMode(TEXT("XRScribe.RunMode"),
 	reinterpret_cast<int32&>(XRScribeRunMode),
 	TEXT("Toggle the mode XRScribe operates in (capture or emulate)"),
@@ -701,7 +701,7 @@ public:
 	{
 		ChainedGetProcAddr = InChainedGetProcAddr;
 
-		EXRScribeRunMode RunModeFromConfig = EXRScribeRunMode::Capture;
+		EXRScribeRunMode RunModeFromConfig = EXRScribeRunMode::Emulate;
 		GConfig->GetInt(TEXT("SystemSettings"), TEXT("XRScribe.RunMode"), (int32&)RunModeFromConfig, GEngineIni);
 
 		if (RunModeFromConfig == EXRScribeRunMode::Capture)
