XRScribe: fix failing API surface tests

Emulation by default broke the API surface tests, which were built assuming capture was the default

#preflgiht 647e562a8549b02876097ef8
#rnx

diff --git a/Engine/Plugins/Runtime/XR/XRScribe/Source/XRScribe/Private/XRScribeAPISurface.cpp b/Engine/Plugins/Runtime/XR/XRScribe/Source/XRScribe/Private/XRScribeAPISurface.cpp
--- a/Engine/Plugins/Runtime/XR/XRScribe/Source/XRScribe/Private/XRScribeAPISurface.cpp
+++ b/Engine/Plugins/Runtime/XR/XRScribe/Source/XRScribe/Private/XRScribeAPISurface.cpp
@@ -697,12 +697,29 @@ public:
 		return ActiveLayer.Get();
 	}
 
+	void SetFallbackRunMode(int32 Fallback) override
+	{
+		if ((Fallback >= (int32)EXRScribeRunMode::Capture) &&
+			(Fallback <= (int32)EXRScribeRunMode::Emulate))
+		{
+			FallbackRunMode = (EXRScribeRunMode)Fallback;
+		}
+		else
+		{
+			UE_LOG(LogXRScribeAPI, Warning, TEXT("Invalid run mode override specified (%d), ignoring"), Fallback);
+		}
+	}
+
 	bool SetChainedGetProcAddr(PFN_xrGetInstanceProcAddr InChainedGetProcAddr) override
 	{
 		ChainedGetProcAddr = InChainedGetProcAddr;
 
 		EXRScribeRunMode RunModeFromConfig = EXRScribeRunMode::Emulate;
-		GConfig->GetInt(TEXT("SystemSettings"), TEXT("XRScribe.RunMode"), (int32&)RunModeFromConfig, GEngineIni);
+		const bool bModeReadFromConfig = GConfig->GetInt(TEXT("SystemSettings"), TEXT("XRScribe.RunMode"), (int32&)RunModeFromConfig, GEngineIni);
+		if (!bModeReadFromConfig)
+		{
+			RunModeFromConfig = FallbackRunMode;
+		}
 
 		if (RunModeFromConfig == EXRScribeRunMode::Capture)
 		{
@@ -751,6 +768,8 @@ private:
 	TSharedPtr<IOpenXRAPILayer, ESPMode::ThreadSafe> EmulationLayer;
 
 	PFN_xrGetInstanceProcAddr ChainedGetProcAddr = nullptr;
+
+	EXRScribeRunMode FallbackRunMode = EXRScribeRunMode::Emulate;
 };
 
 static FOpenXRAPILayerManager OpenXRAPISurface;
diff --git a/Engine/Plugins/Runtime/XR/XRScribe/Source/XRScribe/Public/XRScribeAPISurface.h b/Engine/Plugins/Runtime/XR/XRScribe/Source/XRScribe/Public/XRScribeAPISurface.h
--- a/Engine/Plugins/Runtime/XR/XRScribe/Source/XRScribe/Public/XRScribeAPISurface.h
+++ b/Engine/Plugins/Runtime/XR/XRScribe/Source/XRScribe/Public/XRScribeAPISurface.h
@@ -22,6 +22,8 @@ class IOpenXRAPILayerManager
 public:
 	virtual IOpenXRAPILayer* GetActiveLayer() const = 0;
 
+	virtual void SetFallbackRunMode(int32_t Fallback) = 0;
+
 	virtual bool SetChainedGetProcAddr(PFN_xrGetInstanceProcAddr InChainedGetProcAddr) = 0;
 	
 	static IOpenXRAPILayerManager& Get();
