XRScribe: fix typo in InsertOpenXRAPILayer

Wasn't passing off correct proc addr, whoops

#rnx

diff --git a/Engine/Plugins/Runtime/XR/XRScribe/Source/XRScribe/Private/XRScribeModule.cpp b/Engine/Plugins/Runtime/XR/XRScribe/Source/XRScribe/Private/XRScribeModule.cpp
--- a/Engine/Plugins/Runtime/XR/XRScribe/Source/XRScribe/Private/XRScribeModule.cpp
+++ b/Engine/Plugins/Runtime/XR/XRScribe/Source/XRScribe/Private/XRScribeModule.cpp
@@ -32,7 +32,7 @@ FString FXRScribeModule::GetDisplayName()
 bool FXRScribeModule::InsertOpenXRAPILayer(PFN_xrGetInstanceProcAddr& InOutGetProcAddr)
 {
 	// Ok to pass nullptr for GetProcAddr because that would only disable capture. Replay still valid!
-	const bool bResult = UE::XRScribe::IOpenXRAPILayerManager::Get().SetChainedGetProcAddr(ChainedGetProcAddr);
+	const bool bResult = UE::XRScribe::IOpenXRAPILayerManager::Get().SetChainedGetProcAddr(InOutGetProcAddr);
 
 	if (bResult)
 	{
