XRScribe: proper extension filtering in capture layer

Previously, we'd just passs all extensions in/out thru the capture layer to the runtime. XRScribe does not support the full set of OpenXR extensions yet, so the capture layer has to gate what it exposes to the engine.

Once we receive the loader GetProcAddr, we use it to fetch all instance extensions, and create a subset against the capture layer supported extensions. That subset is then exposed to the engine.

#rnx

diff --git a/Engine/Plugins/Runtime/XR/XRScribe/Source/XRScribe/Private/XRScribeCaptureLayer.cpp b/Engine/Plugins/Runtime/XR/XRScribe/Source/XRScribe/Private/XRScribeCaptureLayer.cpp
--- a/Engine/Plugins/Runtime/XR/XRScribe/Source/XRScribe/Private/XRScribeCaptureLayer.cpp
+++ b/Engine/Plugins/Runtime/XR/XRScribe/Source/XRScribe/Private/XRScribeCaptureLayer.cpp
@@ -3,7 +3,9 @@
 #include "XRScribeCaptureLayer.h"
 
 #include "Containers/StringConv.h"
+#include "HAL/UnrealMemory.h"
 #include "Logging/LogMacros.h"
+#include "Math/UnrealMathUtility.h"
 #include "Misc/FileHelper.h"
 #include "Misc/Paths.h"
 
@@ -30,6 +32,17 @@ void FOpenXRCaptureLayer::SaveCaptureToFile()
 	}
 }
 
+FOpenXRCaptureLayer::FOpenXRCaptureLayer()
+{
+	NativeSupportedExtensionNames.Add(ANSI_TO_TCHAR(XR_KHR_VISIBILITY_MASK_EXTENSION_NAME));
+#if defined(XR_USE_GRAPHICS_API_D3D11)
+	NativeSupportedExtensionNames.Add(ANSI_TO_TCHAR(XR_KHR_D3D11_ENABLE_EXTENSION_NAME));
+#endif
+#if defined(XR_USE_GRAPHICS_API_D3D12)
+	NativeSupportedExtensionNames.Add(ANSI_TO_TCHAR(XR_KHR_D3D12_ENABLE_EXTENSION_NAME));
+#endif
+}
+
 FOpenXRCaptureLayer::~FOpenXRCaptureLayer()
 {
 }
@@ -55,6 +68,38 @@ void FOpenXRCaptureLayer::SetChainedGetProcAddr(PFN_xrGetInstanceProcAddr InChai
 	
 	// TODO: unsure when to actually try to hook this
 	//GetProcAddrHelper(FunctionPassthroughs.GetInstanceProcAddr, XR_NULL_HANDLE, "xrInitializeLoaderKHR", reinterpret_cast<PFN_xrVoidFunction*>(&FunctionPassthroughs.InitializeLoaderKHR));
+
+	{
+		uint32_t NumExtensions = 0;
+		XrResult Result = FunctionPassthroughs.EnumerateInstanceExtensionProperties(nullptr, 0, &NumExtensions, nullptr);
+		ensure(Result == XR_SUCCESS);
+		CaptureEncoder.EncodeEnumerateInstanceExtensionProperties(Result, nullptr, 0, &NumExtensions, nullptr);
+
+
+		TArray<XrExtensionProperties> TotalExtensionProperties;
+		TotalExtensionProperties.AddZeroed(NumExtensions);
+		for (XrExtensionProperties& Property : TotalExtensionProperties)
+		{
+			Property.type = XR_TYPE_EXTENSION_PROPERTIES;
+		}
+		
+		Result = FunctionPassthroughs.EnumerateInstanceExtensionProperties(nullptr, NumExtensions, &NumExtensions, TotalExtensionProperties.GetData());
+		ensure(Result == XR_SUCCESS);
+		CaptureEncoder.EncodeEnumerateInstanceExtensionProperties(Result, nullptr, NumExtensions, &NumExtensions, TotalExtensionProperties.GetData());
+
+		for (const XrExtensionProperties& RuntimeExtensionProperty : TotalExtensionProperties)
+		{
+			if (NativeSupportedExtensionNames.Contains(FString(ANSI_TO_TCHAR(RuntimeExtensionProperty.extensionName))))
+			{
+				PassthruInstanceSupportedExtensions.Add(RuntimeExtensionProperty);
+			}
+		}
+	}
+}
+
+bool FOpenXRCaptureLayer::SupportsInstanceExtension(const ANSICHAR* ExtensionName)
+{
+	return CaptureActiveExtensionNames.Contains(FString(ANSI_TO_TCHAR(ExtensionName)));
 }
 
 // Global
@@ -67,9 +112,40 @@ XrResult FOpenXRCaptureLayer::XrLayerEnumerateApiLayerProperties(uint32_t proper
 
 XrResult FOpenXRCaptureLayer::XrLayerEnumerateInstanceExtensionProperties(const char* layerName, uint32_t propertyCapacityInput, uint32_t* propertyCountOutput, XrExtensionProperties* properties)
 {
-	const XrResult Result = FunctionPassthroughs.EnumerateInstanceExtensionProperties(layerName, propertyCapacityInput, propertyCountOutput, properties);
-	CaptureEncoder.EncodeEnumerateInstanceExtensionProperties(Result, layerName, propertyCapacityInput, propertyCountOutput, properties);
-	return Result;
+	if (layerName != nullptr)
+	{
+		UE_LOG(LogXRScribeCapture, Error, TEXT("Capture layer does not support layer-specific extensions"));
+		return XR_ERROR_RUNTIME_FAILURE;
+	}
+
+	if (propertyCapacityInput == 0)
+	{
+		if (propertyCountOutput == nullptr)
+		{
+			return XR_ERROR_VALIDATION_FAILURE;
+		}
+		*propertyCountOutput = PassthruInstanceSupportedExtensions.Num();
+	}
+	else
+	{
+		if ((propertyCountOutput == nullptr) ||
+			(properties == nullptr))
+		{
+			return XR_ERROR_VALIDATION_FAILURE;
+		}
+		if (propertyCapacityInput < (uint32)PassthruInstanceSupportedExtensions.Num())
+		{
+			return XR_ERROR_SIZE_INSUFFICIENT;
+		}
+
+		const uint32 NumPropertiesToCopy = PassthruInstanceSupportedExtensions.Num();
+		*propertyCountOutput = NumPropertiesToCopy;
+
+		FMemory::Memcpy(properties, PassthruInstanceSupportedExtensions.GetData(), NumPropertiesToCopy * sizeof(XrExtensionProperties));
+
+	}
+
+	return XR_SUCCESS;
 }
 
 XrResult FOpenXRCaptureLayer::XrLayerCreateInstance(const XrInstanceCreateInfo* createInfo, XrInstance* instance)
@@ -130,17 +206,37 @@ XrResult FOpenXRCaptureLayer::XrLayerCreateInstance(const XrInstanceCreateInfo*
 		GetProcAddrHelper(FunctionPassthroughs.GetInstanceProcAddr, *instance, "xrGetInputSourceLocalizedName", reinterpret_cast<PFN_xrVoidFunction*>(&FunctionPassthroughs.GetInputSourceLocalizedName));
 		GetProcAddrHelper(FunctionPassthroughs.GetInstanceProcAddr, *instance, "xrApplyHapticFeedback", reinterpret_cast<PFN_xrVoidFunction*>(&FunctionPassthroughs.ApplyHapticFeedback));
 		GetProcAddrHelper(FunctionPassthroughs.GetInstanceProcAddr, *instance, "xrStopHapticFeedback", reinterpret_cast<PFN_xrVoidFunction*>(&FunctionPassthroughs.StopHapticFeedback));
-		GetProcAddrHelper(FunctionPassthroughs.GetInstanceProcAddr, *instance, "xrGetVisibilityMaskKHR", reinterpret_cast<PFN_xrVoidFunction*>(&FunctionPassthroughs.GetVisibilityMaskKHR));
+
+		// Extensions are managed different than core functions
+		if (createInfo->enabledExtensionCount > 0)
+		{
+			CaptureActiveExtensionNames.Reserve(createInfo->enabledExtensionCount);
+
+			// We could save the full extension info, but the names are fine for now
+			for (uint32 ExtensionIndex = 0; ExtensionIndex < createInfo->enabledExtensionCount; ExtensionIndex++)
+			{
+				CaptureActiveExtensionNames.Add(ANSI_TO_TCHAR(createInfo->enabledExtensionNames[ExtensionIndex]));
+			}
+
+			if (CaptureActiveExtensionNames.Contains(XR_KHR_VISIBILITY_MASK_EXTENSION_NAME))
+			{
+				GetProcAddrHelper(FunctionPassthroughs.GetInstanceProcAddr, *instance, "xrGetVisibilityMaskKHR", reinterpret_cast<PFN_xrVoidFunction*>(&FunctionPassthroughs.GetVisibilityMaskKHR));
+			}
 #if defined(XR_USE_GRAPHICS_API_D3D11)
-		GetProcAddrHelper(FunctionPassthroughs.GetInstanceProcAddr, *instance, "xrGetD3D11GraphicsRequirementsKHR", reinterpret_cast<PFN_xrVoidFunction*>(&FunctionPassthroughs.GetD3D11GraphicsRequirementsKHR));
+			if (CaptureActiveExtensionNames.Contains(XR_KHR_D3D11_ENABLE_EXTENSION_NAME))
+			{
+				GetProcAddrHelper(FunctionPassthroughs.GetInstanceProcAddr, *instance, "xrGetD3D11GraphicsRequirementsKHR", reinterpret_cast<PFN_xrVoidFunction*>(&FunctionPassthroughs.GetD3D11GraphicsRequirementsKHR));
+			}
 #endif
 #if defined(XR_USE_GRAPHICS_API_D3D12)
-		GetProcAddrHelper(FunctionPassthroughs.GetInstanceProcAddr, *instance, "xrGetD3D12GraphicsRequirementsKHR", reinterpret_cast<PFN_xrVoidFunction*>(&FunctionPassthroughs.GetD3D12GraphicsRequirementsKHR));
+			if (CaptureActiveExtensionNames.Contains(XR_KHR_D3D12_ENABLE_EXTENSION_NAME))
+			{
+				GetProcAddrHelper(FunctionPassthroughs.GetInstanceProcAddr, *instance, "xrGetD3D12GraphicsRequirementsKHR", reinterpret_cast<PFN_xrVoidFunction*>(&FunctionPassthroughs.GetD3D12GraphicsRequirementsKHR));
+			}
 #endif
+		}
 	}
 
-	// TODO: might want to fetch the instance properties here in case UE doesn't
-
 	return Result;
 }
 
diff --git a/Engine/Plugins/Runtime/XR/XRScribe/Source/XRScribe/Public/XRScribeCaptureLayer.h b/Engine/Plugins/Runtime/XR/XRScribe/Source/XRScribe/Public/XRScribeCaptureLayer.h
--- a/Engine/Plugins/Runtime/XR/XRScribe/Source/XRScribe/Public/XRScribeCaptureLayer.h
+++ b/Engine/Plugins/Runtime/XR/XRScribe/Source/XRScribe/Public/XRScribeCaptureLayer.h
@@ -6,21 +6,21 @@
 #include "XRScribeAPILayer.h"
 #include "XRScribeAPIPassthrough.h"
 
+#include "Containers/Array.h"
+#include "Containers/UnrealString.h"
+
 namespace UE::XRScribe
 {
 
 class FOpenXRCaptureLayer : public IOpenXRAPILayer
 {
 public:
+	FOpenXRCaptureLayer();
 	virtual ~FOpenXRCaptureLayer() override;
 
 	virtual void SetChainedGetProcAddr(PFN_xrGetInstanceProcAddr InChainedGetProcAddr) override;
 
-	virtual bool SupportsInstanceExtension(const ANSICHAR* ExtensionName) override
-	{
-		// TODO: we need to stash off the extensions from CreateInstance to validate these requests
-		return true;
-	}
+	virtual bool SupportsInstanceExtension(const ANSICHAR* ExtensionName) override;
 
 	// Global
 	virtual XrResult XrLayerEnumerateApiLayerProperties(uint32_t propertyCapacityInput, uint32_t* propertyCountOutput, XrApiLayerProperties* properties) override;
@@ -115,6 +115,10 @@ protected:
 
 	FOpenXRAPIPassthrough FunctionPassthroughs;
 	FOpenXRCaptureEncoder CaptureEncoder;
+
+	TArray<FString> NativeSupportedExtensionNames;
+	TArray<XrExtensionProperties> PassthruInstanceSupportedExtensions;
+	TArray<FString> CaptureActiveExtensionNames;
 };
 
 } // namespace UE::XRScribe
