XRScribe: built-in captures for emulation

The expected core use case for the tool is emulation. To simplify most uses of emulation, we can add some default captures to allow for quick emulation with needing a capture file from an external source.
This will be controlled by a CLI arg, tho in the future, it could come from a config.

#rnx

diff --git a/Engine/Plugins/Runtime/XR/XRScribe/Source/XRScribe/Private/XRScribeAPISurface.cpp b/Engine/Plugins/Runtime/XR/XRScribe/Source/XRScribe/Private/XRScribeAPISurface.cpp
--- a/Engine/Plugins/Runtime/XR/XRScribe/Source/XRScribe/Private/XRScribeAPISurface.cpp
+++ b/Engine/Plugins/Runtime/XR/XRScribe/Source/XRScribe/Private/XRScribeAPISurface.cpp
@@ -4,8 +4,12 @@
 
 #include "Containers/Map.h"
 #include "HAL/IConsoleManager.h"
+#include "Interfaces/IPluginManager.h"
 #include "Logging/LogMacros.h"
+#include "Misc/CommandLine.h"
 #include "Misc/ConfigCacheIni.h"
+#include "Misc/Parse.h"
+#include "Misc/Paths.h"
 #include "OpenXRCore.h"
 #include "Templates/SharedPointer.h"
 #include "XRScribeCaptureLayer.h"
@@ -681,6 +685,47 @@ void ClearFunctionMaps()
 	InstanceFunctionMap.Reset();
 }
 
+EXRScribeRunMode DetermineRunMode(EXRScribeRunMode FallbackRunMode, FString& EmulationLoadPath)
+{
+	// TODO: load file from plugin settings
+
+	EXRScribeRunMode RunModeFromConfig = EXRScribeRunMode::Emulate;
+	const bool bModeReadFromConfig = GConfig->GetInt(TEXT("SystemSettings"), TEXT("XRScribe.RunMode"), (int32&)RunModeFromConfig, GEngineIni);
+	if (!bModeReadFromConfig)
+	{
+		RunModeFromConfig = FallbackRunMode;
+	}
+
+	bool bEmulationCLI = false;
+	EmulationLoadPath = FPaths::ProjectSavedDir() / TEXT("Capture.xrs");
+
+	FString RequestedEmulationProfile;
+	bEmulationCLI = FParse::Value(FCommandLine::Get(), TEXT("-xremu="), RequestedEmulationProfile);
+	if (bEmulationCLI)
+	{
+		UE_LOG(LogXRScribeAPI, Log, TEXT("xremu command-line argument detected to emulate built-in capture"));
+
+		FString ExtrasDir = IPluginManager::Get().FindPlugin("XRScribe")->GetBaseDir() / TEXT("Extras");
+		if (RequestedEmulationProfile == TEXT("mq2"))
+		{
+			EmulationLoadPath = ExtrasDir / TEXT("mq2.xrs");
+			UE_LOG(LogXRScribeAPI, Log, TEXT("mq2 built-in capture requested"));
+		}
+		else if (RequestedEmulationProfile == TEXT("vi"))
+		{
+			EmulationLoadPath = ExtrasDir / TEXT("vi.xrs");
+			UE_LOG(LogXRScribeAPI, Log, TEXT("vi built-in capture requested"));
+		}
+		else
+		{
+			UE_LOG(LogXRScribeAPI, Warning, TEXT("Unknown built-in capture, deferring to config run-mode"));
+			bEmulationCLI = false;
+		}
+	}
+
+	return bEmulationCLI ? EXRScribeRunMode::Emulate : RunModeFromConfig;
+}
+
 class FOpenXRAPILayerManager : public IOpenXRAPILayerManager
 {
 public:
@@ -714,14 +759,10 @@ public:
 	{
 		ChainedGetProcAddr = InChainedGetProcAddr;
 
-		EXRScribeRunMode RunModeFromConfig = EXRScribeRunMode::Emulate;
-		const bool bModeReadFromConfig = GConfig->GetInt(TEXT("SystemSettings"), TEXT("XRScribe.RunMode"), (int32&)RunModeFromConfig, GEngineIni);
-		if (!bModeReadFromConfig)
-		{
-			RunModeFromConfig = FallbackRunMode;
-		}
+		FString EmulationLoadPath;
+		const EXRScribeRunMode SelectedRunMode = DetermineRunMode(FallbackRunMode, EmulationLoadPath);
 
-		if (RunModeFromConfig == EXRScribeRunMode::Capture)
+		if (SelectedRunMode == EXRScribeRunMode::Capture)
 		{
 			UE_LOG(LogXRScribeAPI, Log, TEXT("Capture layer selected"));
 
@@ -738,11 +779,11 @@ public:
 				return false;
 			}
 		}
-		else if (RunModeFromConfig == EXRScribeRunMode::Emulate)
+		else if (SelectedRunMode == EXRScribeRunMode::Emulate)
 		{
 			UE_LOG(LogXRScribeAPI, Log, TEXT("Emulation layer selected"));
 			ActiveLayer = EmulationLayer;
-			const bool bEmulationLoadedCapture = reinterpret_cast<FOpenXREmulationLayer*>(EmulationLayer.Get())->LoadCaptureFromFile();
+			const bool bEmulationLoadedCapture = reinterpret_cast<FOpenXREmulationLayer*>(EmulationLayer.Get())->LoadCaptureFromFile(EmulationLoadPath);
 
 			if (bEmulationLoadedCapture)
 			{
diff --git a/Engine/Plugins/Runtime/XR/XRScribe/Source/XRScribe/Private/XRScribeEmulationLayer.cpp b/Engine/Plugins/Runtime/XR/XRScribe/Source/XRScribe/Private/XRScribeEmulationLayer.cpp
--- a/Engine/Plugins/Runtime/XR/XRScribe/Source/XRScribe/Private/XRScribeEmulationLayer.cpp
+++ b/Engine/Plugins/Runtime/XR/XRScribe/Source/XRScribe/Private/XRScribeEmulationLayer.cpp
@@ -10,7 +10,6 @@
 #include "Logging/LogMacros.h"
 #include "Misc/DateTime.h"
 #include "Misc/FileHelper.h"
-#include "Misc/Paths.h"
 #include "OpenXRCore.h"
 #include "RHICommandList.h"
 #include "RHIResources.h"
@@ -247,13 +246,11 @@ bool FOpenXREmulationLayer::SupportsInstanceExtension(const ANSICHAR* ExtensionN
 	});
 }
 
-bool FOpenXREmulationLayer::LoadCaptureFromFile()
+bool FOpenXREmulationLayer::LoadCaptureFromFile(const FString& EmulationLoadPath)
 {
-	// TODO: load file from plugin settings
-	const FString CaptureFilePath = FPaths::ProjectSavedDir() / TEXT("Capture.xrs");
-	if (FFileHelper::LoadFileToArray(CaptureDecoder.GetEncodedData(), *CaptureFilePath))
+	if (FFileHelper::LoadFileToArray(CaptureDecoder.GetEncodedData(), *EmulationLoadPath))
 	{
-		UE_LOG(LogXRScribeEmulate, Log, TEXT("Capture successfully loaded: %s"), *CaptureFilePath);
+		UE_LOG(LogXRScribeEmulate, Log, TEXT("Capture successfully loaded: %s"), *EmulationLoadPath);
 		CaptureDecoder.DecodeDataFromMemory();
 	}
 	else
diff --git a/Engine/Plugins/Runtime/XR/XRScribe/Source/XRScribe/Public/XRScribeEmulationLayer.h b/Engine/Plugins/Runtime/XR/XRScribe/Source/XRScribe/Public/XRScribeEmulationLayer.h
--- a/Engine/Plugins/Runtime/XR/XRScribe/Source/XRScribe/Public/XRScribeEmulationLayer.h
+++ b/Engine/Plugins/Runtime/XR/XRScribe/Source/XRScribe/Public/XRScribeEmulationLayer.h
@@ -30,7 +30,7 @@ public:
 
 	virtual bool SupportsInstanceExtension(const ANSICHAR* ExtensionName) override;
 
-	bool LoadCaptureFromFile();
+	bool LoadCaptureFromFile(const FString& EmulationLoadPath);
 	bool LoadCaptureFromData(const TArray<uint8>& EncodedData);
 
 	// Global
diff --git a/Engine/Plugins/Runtime/XR/XRScribe/Source/XRScribe/XRScribe.Build.cs b/Engine/Plugins/Runtime/XR/XRScribe/Source/XRScribe/XRScribe.Build.cs
--- a/Engine/Plugins/Runtime/XR/XRScribe/Source/XRScribe/XRScribe.Build.cs
+++ b/Engine/Plugins/Runtime/XR/XRScribe/Source/XRScribe/XRScribe.Build.cs
@@ -23,6 +23,7 @@ namespace UnrealBuildTool.Rules
 			PrivateDependencyModuleNames.AddRange(new string[]
 			{
                 "DeveloperSettings",
+				"Projects",
 			});
 
 			// To get access to OpenXRPlatformRHI
