XRScribe: proper multi-threaded write protection for encoder

Originally, XRScribe would manually serialize blobs of data, which meant the previous Serialize override worked fine.
Once XRScribe used the full FArchive functionality, the old mutex locking wasn't locking every encoder write, leading to multi-threading bugs.

#rnx

diff --git a/Engine/Plugins/Runtime/XR/XRScribe/Source/XRScribe/Private/XRScribeAPIEncoder.cpp b/Engine/Plugins/Runtime/XR/XRScribe/Source/XRScribe/Private/XRScribeAPIEncoder.cpp
--- a/Engine/Plugins/Runtime/XR/XRScribe/Source/XRScribe/Private/XRScribeAPIEncoder.cpp
+++ b/Engine/Plugins/Runtime/XR/XRScribe/Source/XRScribe/Private/XRScribeAPIEncoder.cpp
@@ -18,11 +18,13 @@
 {
 }
 
-void FOpenXRCaptureEncoder::Serialize(void* Data, int64 Num)
+template <typename PacketType>
+void FOpenXRCaptureEncoder::WritePacketData(PacketType& Data)
 {
-	// TODO Would be cool to kick this to a thread-safe queue so we don't have to lock!
+	// TODO Support piping data to a thread-safe queue in order to avoid locking
+
 	FWriteScopeLock Lock(CaptureWriteMutex);
-	TMemoryWriter<32>::Serialize(Data, Num);
+	*this << Data;
 }
 
 uint32 NumElementsToWrite(uint32 CapacityCount, const uint32* GeneratedCount)
@@ -39,7 +41,7 @@
 
 	Data.LayerProperties.Append(properties, NumElementsToWrite(propertyCapacityInput, propertyCountOutput));
 
-	*this << Data;
+	WritePacketData(Data);
 }
 
 void FOpenXRCaptureEncoder::EncodeEnumerateInstanceExtensionProperties(const XrResult Result, const char* layerName,
@@ -59,7 +61,7 @@
 
 	Data.ExtensionProperties.Append(properties, NumElementsToWrite(propertyCapacityInput, propertyCountOutput));
 
-	*this << Data;
+	WritePacketData(Data);
 }
 
 void FOpenXRCaptureEncoder::EncodeCreateInstance(const XrResult Result, const XrInstanceCreateInfo* createInfo, 
@@ -94,7 +96,7 @@
 
 	Data.GeneratedInstance = *instance;
 
-	*this << Data;
+	WritePacketData(Data);
 }
 
 void FOpenXRCaptureEncoder::EncodeDestroyInstance(const XrResult Result, const XrInstance instance)
@@ -104,7 +106,7 @@
 
 	Data.Instance = instance;
 
-	*this << Data;
+	WritePacketData(Data);
 }
 
 void FOpenXRCaptureEncoder::EncodeGetInstanceProperties(const XrResult Result, const XrInstance instance,
@@ -116,7 +118,7 @@
 	Data.Instance = instance;
 	Data.InstanceProperties = *instanceProperties;
 
-	*this << Data;
+	WritePacketData(Data);
 }
 
 void FOpenXRCaptureEncoder::EncodePollEvent(const XrResult Result, const XrInstance instance, const XrEventDataBuffer* eventData)
@@ -147,7 +149,7 @@
 	Data.SystemGetInfo = *getInfo;
 	Data.SystemId = *systemId;
 
-	*this << Data;
+	WritePacketData(Data);
 }
 
 void FOpenXRCaptureEncoder::EncodeGetSystemProperties(const XrResult Result, const XrInstance instance, const XrSystemId systemId, const XrSystemProperties* properties)
@@ -159,7 +161,7 @@
 	Data.SystemId = systemId;
 	Data.SystemProperties = *properties;
 
-	*this << Data;
+	WritePacketData(Data);
 }
 
 void FOpenXRCaptureEncoder::EncodeEnumerateEnvironmentBlendModes(const XrResult Result, const XrInstance instance, const XrSystemId systemId, XrViewConfigurationType viewConfigurationType, const uint32_t environmentBlendModeCapacityInput, const uint32_t* environmentBlendModeCountOutput, const XrEnvironmentBlendMode* environmentBlendModes)
@@ -173,7 +175,7 @@
 
 	Data.EnvironmentBlendModes.Append(environmentBlendModes, NumElementsToWrite(environmentBlendModeCapacityInput, environmentBlendModeCountOutput));
 
-	*this << Data;
+	WritePacketData(Data);
 }
 
 void FOpenXRCaptureEncoder::EncodeCreateSession(const XrResult Result, const XrInstance instance, const XrSessionCreateInfo* createInfo, const XrSession* session)
@@ -185,7 +187,7 @@
 	Data.SessionCreateInfo = *createInfo;
 	Data.Session = *session;
 
-	*this << Data;
+	WritePacketData(Data);
 }
 
 void FOpenXRCaptureEncoder::EncodeDestroySession(const XrResult Result, const XrSession session)
@@ -195,7 +197,7 @@
 
 	Data.Session = session;
 
-	*this << Data;
+	WritePacketData(Data);
 }
 
 void FOpenXRCaptureEncoder::EncodeEnumerateReferenceSpaces(const XrResult Result, const XrSession session, const uint32_t spaceCapacityInput, const uint32_t* spaceCountOutput, const XrReferenceSpaceType* spaces)
@@ -207,7 +209,7 @@
 
 	Data.Spaces.Append(spaces, NumElementsToWrite(spaceCapacityInput, spaceCountOutput));
 
-	*this << Data;
+	WritePacketData(Data);
 }
 
 void FOpenXRCaptureEncoder::EncodeCreateReferenceSpace(const XrResult Result, const XrSession session, const XrReferenceSpaceCreateInfo* createInfo, const XrSpace* space)
@@ -219,7 +221,7 @@
 	Data.ReferenceSpaceCreateInfo = *createInfo;
 	Data.Space = *space;
 
-	*this << Data;
+	WritePacketData(Data);
 }
 
 void FOpenXRCaptureEncoder::EncodeGetReferenceSpaceBoundsRect(const XrResult Result, const XrSession session, const XrReferenceSpaceType referenceSpaceType, const XrExtent2Df* bounds)
@@ -231,7 +233,7 @@
 	Data.ReferenceSpaceType = referenceSpaceType;
 	Data.Bounds = *bounds;
 
-	*this << Data;
+	WritePacketData(Data);
 }
 
 void FOpenXRCaptureEncoder::EncodeCreateActionSpace(const XrResult Result, const XrSession session, const XrActionSpaceCreateInfo* createInfo, const XrSpace* space)
@@ -243,7 +245,7 @@
 	Data.ActionSpaceCreateInfo = *createInfo;
 	Data.Space = *space;
 
-	*this << Data;
+	WritePacketData(Data);
 }
 
 void FOpenXRCaptureEncoder::EncodeLocateSpace(const XrResult Result, const XrSpace space, const XrSpace baseSpace, const XrTime time, const XrSpaceLocation* location)
@@ -256,7 +258,7 @@
 	Data.Time = time;
 	Data.Location = *location;
 
-	*this << Data;
+	WritePacketData(Data);
 }
 
 void FOpenXRCaptureEncoder::EncodeDestroySpace(const XrResult Result, const XrSpace space)
@@ -266,7 +268,7 @@
 
 	Data.Space = space;
 
-	*this << Data;
+	WritePacketData(Data);
 }
 
 void FOpenXRCaptureEncoder::EncodeEnumerateViewConfigurations(const XrResult Result, const XrInstance instance, const XrSystemId systemId, const uint32_t viewConfigurationTypeCapacityInput, const uint32_t* viewConfigurationTypeCountOutput, const XrViewConfigurationType* viewConfigurationTypes)
@@ -278,7 +280,7 @@
 	Data.SystemId = systemId;
 	Data.ViewConfigurationTypes.Append(viewConfigurationTypes, NumElementsToWrite(viewConfigurationTypeCapacityInput, viewConfigurationTypeCountOutput));
 
-	*this << Data;
+	WritePacketData(Data);
 }
 
 void FOpenXRCaptureEncoder::EncodeGetViewConfigurationProperties(const XrResult Result, const XrInstance instance, const XrSystemId systemId, const XrViewConfigurationType viewConfigurationType, const XrViewConfigurationProperties* configurationProperties)
@@ -291,7 +293,7 @@
 	Data.ViewConfigurationType = viewConfigurationType;
 	Data.ConfigurationProperties = *configurationProperties;
 
-	*this << Data;
+	WritePacketData(Data);
 }
 
 void FOpenXRCaptureEncoder::EncodeEnumerateViewConfigurationViews(const XrResult Result, const XrInstance instance, const XrSystemId systemId, const XrViewConfigurationType viewConfigurationType, const uint32_t viewCapacityInput, const uint32_t* viewCountOutput, const XrViewConfigurationView* views)
@@ -304,7 +306,7 @@
 	Data.ViewConfigurationType = viewConfigurationType;
 	Data.Views.Append(views, NumElementsToWrite(viewCapacityInput, viewCountOutput));
 
-	*this << Data;
+	WritePacketData(Data);
 }
 
 void FOpenXRCaptureEncoder::EncodeEnumerateSwapchainFormats(const XrResult Result, const XrSession session, const uint32_t formatCapacityInput, const uint32_t* formatCountOutput, int64_t* formats)
@@ -315,7 +317,7 @@
 	Data.Session = session;
 	Data.Formats.Append(formats, NumElementsToWrite(formatCapacityInput, formatCountOutput));
 
-	*this << Data;
+	WritePacketData(Data);
 }
 
 void FOpenXRCaptureEncoder::EncodeCreateSwapchain(const XrResult Result, const XrSession session, const XrSwapchainCreateInfo* createInfo, const XrSwapchain* swapchain)
@@ -327,7 +329,7 @@
 	Data.SwapchainCreateInfo = *createInfo;
 	Data.Swapchain = *swapchain;
 
-	*this << Data;
+	WritePacketData(Data);
 }
 
 void FOpenXRCaptureEncoder::EncodeDestroySwapchain(const XrResult Result, const XrSwapchain swapchain)
@@ -337,7 +339,7 @@
 
 	Data.Swapchain = swapchain;
 
-	*this << Data;
+	WritePacketData(Data);
 }
 
 void FOpenXRCaptureEncoder::EncodeEnumerateSwapchainImages(const XrResult Result, const XrSwapchain swapchain, const uint32_t imageCapacityInput, const uint32_t* imageCountOutput, const XrSwapchainImageBaseHeader* images)
@@ -351,7 +353,7 @@
 	// TODO: We actually need to check the real type here, and enumerate that
 	// I guess the 'real' image might not be super important for emulation
 
-	*this << Data;
+	WritePacketData(Data);
 }
 
 void FOpenXRCaptureEncoder::EncodeAcquireSwapchainImage(const XrResult Result, const XrSwapchain swapchain, const XrSwapchainImageAcquireInfo* acquireInfo, const uint32_t* index)
@@ -363,7 +365,7 @@
 	Data.AcquireInfo = *acquireInfo;
 	Data.Index = *index;
 
-	*this << Data;
+	WritePacketData(Data);
 }
 
 void FOpenXRCaptureEncoder::EncodeWaitSwapchainImage(const XrResult Result, const XrSwapchain swapchain, const XrSwapchainImageWaitInfo* waitInfo)
@@ -374,7 +376,7 @@
 	Data.Swapchain = swapchain;
 	Data.WaitInfo = *waitInfo;
 
-	*this << Data;
+	WritePacketData(Data);
 }
 
 void FOpenXRCaptureEncoder::EncodeReleaseSwapchainImage(const XrResult Result, const XrSwapchain swapchain, const XrSwapchainImageReleaseInfo* releaseInfo)
@@ -385,7 +387,7 @@
 	Data.Swapchain = swapchain;
 	Data.ReleaseInfo = *releaseInfo;
 
-	*this << Data;
+	WritePacketData(Data);
 }
 
 void FOpenXRCaptureEncoder::EncodeBeginSession(const XrResult Result, const XrSession session, const XrSessionBeginInfo* beginInfo)
@@ -396,7 +398,7 @@
 	Data.Session = session;
 	Data.SessionBeginInfo = *beginInfo;
 
-	*this << Data;
+	WritePacketData(Data);
 }
 
 void FOpenXRCaptureEncoder::EncodeEndSession(const XrResult Result, const XrSession session)
@@ -406,7 +408,7 @@
 
 	Data.Session = session;
 
-	*this << Data;
+	WritePacketData(Data);
 }
 
 void FOpenXRCaptureEncoder::EncodeRequestExitSession(const XrResult Result, const XrSession session)
@@ -416,7 +418,7 @@
 
 	Data.Session = session;
 
-	*this << Data;
+	WritePacketData(Data);
 }
 
 void FOpenXRCaptureEncoder::EncodeWaitFrame(const XrResult Result, const XrSession session, const XrFrameWaitInfo* frameWaitInfo, const XrFrameState* frameState)
@@ -428,7 +430,7 @@
 	Data.FrameWaitInfo = *frameWaitInfo;
 	Data.FrameState = *frameState;
 
-	*this << Data;
+	WritePacketData(Data);
 }
 
 void FOpenXRCaptureEncoder::EncodeBeginFrame(const XrResult Result, const XrSession session, const XrFrameBeginInfo* frameBeginInfo)
@@ -439,7 +441,7 @@
 	Data.Session = session;
 	Data.FrameBeginInfo = *frameBeginInfo;
 
-	*this << Data;
+	WritePacketData(Data);
 }
 void FOpenXRCaptureEncoder::EncodeEndFrame(const XrResult Result, const XrSession session, const XrFrameEndInfo* frameEndInfo)
 {
@@ -484,7 +486,7 @@
 	Data.ProjectionLayers = MoveTemp(ProjectionLayers);
 	Data.ProjectionViews = MoveTemp(ProjectionViews);
 
-	*this << Data;
+	WritePacketData(Data);
 }
 
 void FOpenXRCaptureEncoder::EncodeLocateViews(const XrResult Result, const XrSession session, const XrViewLocateInfo* viewLocateInfo, const XrViewState* viewState, const uint32_t viewCapacityInput, const uint32_t* viewCountOutput, const XrView* views)
@@ -497,7 +499,7 @@
 	Data.ViewState = *viewState;
 	Data.Views.Append(views, NumElementsToWrite(viewCapacityInput, viewCountOutput));
 
-	*this << Data;
+	WritePacketData(Data);
 }
 
 void FOpenXRCaptureEncoder::EncodeStringToPath(const XrResult Result, const XrInstance instance, const char* pathString, const XrPath* path)
@@ -509,7 +511,7 @@
 	FCStringAnsi::Strncpy(Data.PathStringToWrite.GetData(), pathString, XR_MAX_PATH_LENGTH);
 	Data.GeneratedPath = *path;
 
-	*this << Data;
+	WritePacketData(Data);
 }
 
 void FOpenXRCaptureEncoder::EncodePathToString(const XrResult Result, const XrInstance instance, const XrPath path, const uint32_t bufferCapacityInput, const uint32_t* bufferCountOutput, const char* buffer)
@@ -526,7 +528,7 @@
 		FCStringAnsi::Strncpy(Data.PathStringToRead.GetData(), buffer, MaxBufferLen);
 	}
 
-	*this << Data;
+	WritePacketData(Data);
 }
 
 void FOpenXRCaptureEncoder::EncodeCreateActionSet(const XrResult Result, const XrInstance instance, const XrActionSetCreateInfo* createInfo, const XrActionSet* actionSet)
@@ -538,7 +540,7 @@
 	Data.ActionSetCreateInfo = *createInfo;
 	Data.ActionSet = *actionSet;
 
-	*this << Data;
+	WritePacketData(Data);
 }
 
 void FOpenXRCaptureEncoder::EncodeDestroyActionSet(const XrResult Result, const XrActionSet actionSet)
@@ -548,7 +550,7 @@
 
 	Data.ActionSet = actionSet;
 
-	*this << Data;
+	WritePacketData(Data);
 }
 
 void FOpenXRCaptureEncoder::EncodeCreateAction(const XrResult Result, const XrActionSet actionSet, const XrActionCreateInfo* createInfo, const XrAction* action)
@@ -566,7 +568,7 @@
 	FCStringAnsi::Strncpy(Data.LocalizedActionName.GetData(), createInfo->localizedActionName, XR_MAX_LOCALIZED_ACTION_NAME_SIZE);
 	Data.Action = *action;
 
-	*this << Data;
+	WritePacketData(Data);
 }
 
 void FOpenXRCaptureEncoder::EncodeDestroyAction(const XrResult Result, const XrAction action)
@@ -576,7 +578,7 @@
 
 	Data.Action = action;
 
-	*this << Data;
+	WritePacketData(Data);
 }
 
 void FOpenXRCaptureEncoder::EncodeSuggestInteractionProfileBindings(const XrResult Result, const XrInstance instance, const XrInteractionProfileSuggestedBinding* suggestedBindings)
@@ -591,7 +593,7 @@
 		Data.SuggestedBindings.Append(suggestedBindings->suggestedBindings, suggestedBindings->countSuggestedBindings);
 	}
 
-	*this << Data;
+	WritePacketData(Data);
 }
 
 void FOpenXRCaptureEncoder::EncodeAttachSessionActionSets(const XrResult Result, const XrSession session, const XrSessionActionSetsAttachInfo* attachInfo)
@@ -605,7 +607,7 @@
 		Data.ActionSets.Append(attachInfo->actionSets, attachInfo->countActionSets);
 	}
 
-	*this << Data;
+	WritePacketData(Data);
 }
 
 void FOpenXRCaptureEncoder::EncodeGetCurrentInteractionProfile(const XrResult Result, const XrSession session, const XrPath topLevelUserPath, const XrInteractionProfileState* interactionProfile)
@@ -617,7 +619,7 @@
 	Data.TopLevelUserPath = topLevelUserPath;
 	Data.InteractionProfile = *interactionProfile;
 
-	*this << Data;
+	WritePacketData(Data);
 }
 
 void FOpenXRCaptureEncoder::EncodeGetActionStateBoolean(const XrResult Result, const XrSession session, const XrActionStateGetInfo* getInfo, const XrActionStateBoolean* state)
@@ -629,7 +631,7 @@
 	Data.GetInfoBoolean = *getInfo;
 	Data.BooleanState = *state;
 
-	*this << Data;
+	WritePacketData(Data);
 }
 
 void FOpenXRCaptureEncoder::EncodeGetActionStateFloat(const XrResult Result, const XrSession session, const XrActionStateGetInfo* getInfo, const XrActionStateFloat* state)
@@ -641,7 +643,7 @@
 	Data.GetInfoFloat = *getInfo;
 	Data.FloatState = *state;
 
-	*this << Data;
+	WritePacketData(Data);
 }
 void FOpenXRCaptureEncoder::EncodeGetActionStateVector2f(const XrResult Result, const XrSession session, const XrActionStateGetInfo* getInfo, const XrActionStateVector2f* state)
 {
@@ -652,7 +654,7 @@
 	Data.GetInfoVector2f = *getInfo;
 	Data.Vector2fState = *state;
 
-	*this << Data;
+	WritePacketData(Data);
 }
 void FOpenXRCaptureEncoder::EncodeGetActionStatePose(const XrResult Result, const XrSession session, const XrActionStateGetInfo* getInfo, const XrActionStatePose* state)
 {
@@ -663,7 +665,7 @@
 	Data.GetInfoPose = *getInfo;
 	Data.PoseState = *state;
 
-	*this << Data;
+	WritePacketData(Data);
 }
 
 void FOpenXRCaptureEncoder::EncodeSyncActions(const XrResult Result, const XrSession session, const	XrActionsSyncInfo* syncInfo)
@@ -677,7 +679,7 @@
 		Data.ActiveActionSets.Append(syncInfo->activeActionSets, syncInfo->countActiveActionSets);
 	}
 
-	*this << Data;
+	WritePacketData(Data);
 }
 
 void FOpenXRCaptureEncoder::EncodeEnumerateBoundSourcesForAction(const XrResult Result, const XrSession session, const XrBoundSourcesForActionEnumerateInfo* enumerateInfo, const uint32_t sourceCapacityInput, const uint32_t* sourceCountOutput, const XrPath* sources)
@@ -694,7 +696,7 @@
 	//	Data.Sources.Append(sources, *sourceCountOutput);
 	//}
 
-	//*this << Data;
+	//WritePacketData(Data);
 }
 
 void FOpenXRCaptureEncoder::EncodeGetInputSourceLocalizedName(const XrResult Result, const XrSession session, const XrInputSourceLocalizedNameGetInfo* getInfo, const uint32_t bufferCapacityInput, const uint32_t* bufferCountOutput, const char* buffer)
@@ -712,7 +714,7 @@
 	//	FCStringAnsi::Strncpy(Data.LocalizedName.GetData(), buffer, MaxBufferLen);
 	//}
 
-	//*this << Data;
+	//WritePacketData(Data);
 }
 
 void FOpenXRCaptureEncoder::EncodeApplyHapticFeedback(const XrResult Result, const XrSession session, const XrHapticActionInfo* hapticActionInfo, const XrHapticBaseHeader* hapticFeedback)
@@ -726,7 +728,7 @@
 
 	// TODO: We might want to record all the feedback in the future, but its probably good enough to acknowledge the API being called
 
-	*this << Data;
+	WritePacketData(Data);
 }
 
 void FOpenXRCaptureEncoder::EncodeStopHapticFeedback(const XrResult Result, const XrSession session, const XrHapticActionInfo* hapticActionInfo)
@@ -737,7 +739,7 @@
 	Data.Session = session;
 	Data.HapticActionInfo = *hapticActionInfo;
 
-	*this << Data;
+	WritePacketData(Data);
 }
 
 // XR_KHR_loader_init
@@ -750,7 +752,7 @@
 
 	check(0);
 
-	*this << Data;
+	WritePacketData(Data);
 }
 
 // XR_KHR_visibility_mask
@@ -767,7 +769,7 @@
 	Data.Vertices.Append(visibilityMask->vertices, NumElementsToWrite(visibilityMask->vertexCapacityInput, &visibilityMask->vertexCountOutput));
 	Data.Indices.Append(visibilityMask->indices, NumElementsToWrite(visibilityMask->indexCapacityInput, &visibilityMask->indexCountOutput));
 
-	*this << Data;
+	WritePacketData(Data);
 }
 
 #if defined(XR_USE_GRAPHICS_API_D3D11)
@@ -780,7 +782,7 @@
 	Data.SystemId = systemId;
 	Data.GraphicsRequirementsD3D11 = *graphicsRequirements;
 
-	*this << Data;
+	WritePacketData(Data);
 }
 #endif
 
@@ -794,7 +796,7 @@
 	Data.SystemId = systemId;
 	Data.GraphicsRequirementsD3D12 = *graphicsRequirements;
 
-	*this << Data;
+	WritePacketData(Data);
 }
 #endif
 
diff --git a/Engine/Plugins/Runtime/XR/XRScribe/Source/XRScribe/Public/XRScribeAPIEncoder.h b/Engine/Plugins/Runtime/XR/XRScribe/Source/XRScribe/Public/XRScribeAPIEncoder.h
--- a/Engine/Plugins/Runtime/XR/XRScribe/Source/XRScribe/Public/XRScribeAPIEncoder.h
+++ b/Engine/Plugins/Runtime/XR/XRScribe/Source/XRScribe/Public/XRScribeAPIEncoder.h
@@ -17,8 +17,6 @@ public:
 	FOpenXRCaptureEncoder();
 	~FOpenXRCaptureEncoder();
 
-	virtual void Serialize(void* Data, int64 Num) override;
-
 	// Core
 
 	// Global
@@ -105,10 +103,11 @@ public:
 	void EncodeGetD3D12GraphicsRequirementsKHR(const XrResult Result, const XrInstance instance, const XrSystemId systemId, const XrGraphicsRequirementsD3D12KHR* graphicsRequirements);
 #endif
 
-	FRWLock& GetDataLock() { return CaptureWriteMutex; }
-
 protected:
 
+	template <typename PacketType>
+	void WritePacketData(PacketType& Data);
+
 	FRWLock CaptureWriteMutex;
 };
 
