XRScribe: fix plugin metadata

#rnx

diff --git a/Engine/Plugins/Runtime/XR/XRScribe/XRScribe.uplugin b/Engine/Plugins/Runtime/XR/XRScribe/XRScribe.uplugin
--- a/Engine/Plugins/Runtime/XR/XRScribe/XRScribe.uplugin
+++ b/Engine/Plugins/Runtime/XR/XRScribe/XRScribe.uplugin
@@ -1,10 +1,10 @@
 {
 	"FileVersion": 3,
 	"Version": 1,
-	"VersionName": "1.0",
+	"VersionName": "0.1",
 	"FriendlyName": "XRScribe",
-	"Description": "OpenXR API Record/Playback Plugin",
-	"Category": "Experimental",
+	"Description": "OpenXR API Capture/Emulation",
+	"Category": "Virtual Reality",
 	"CreatedBy": "Epic Games, Inc.",
 	"CreatedByURL": "https://epicgames.com",
 	"DocsURL": "",
@@ -12,7 +12,8 @@
 	"SupportURL": "",
 	"EnabledByDefault": false,
 	"CanContainContent": false,
-	"IsBetaVersion": true,
+	"IsBetaVersion": false,
+	"IsExperimentalVersion": true,
 	"Installed": false,
 	"SupportedTargetPlatforms": [
 		"Win64"
