XRScribe: pose replay matches original timescale

Instead of playing poses back as fast as possible, XRScribe now re-samples the poses to 120 Hz.
Emulation then samples poses based on a time offset from the session start.
This allows for a emulation-framerate-independent replay.

diff --git a/Engine/Plugins/Runtime/XR/XRScribe/Source/XRScribe/Private/XRScribeAPIDecoder.cpp b/Engine/Plugins/Runtime/XR/XRScribe/Source/XRScribe/Private/XRScribeAPIDecoder.cpp
--- a/Engine/Plugins/Runtime/XR/XRScribe/Source/XRScribe/Private/XRScribeAPIDecoder.cpp
+++ b/Engine/Plugins/Runtime/XR/XRScribe/Source/XRScribe/Private/XRScribeAPIDecoder.cpp
@@ -513,7 +513,8 @@ void FOpenXRCaptureDecoder::DecodeWaitFrame(const FOpenXRAPIPacketBase& BasePack
 	FOpenXRWaitFramePacket Data(XrResult::XR_ERROR_RUNTIME_FAILURE);
 
 	EncodedData << Data;
-	// nothing to do yet
+
+	WaitFrames.Add(Data);
 }
 
 void FOpenXRCaptureDecoder::DecodeBeginFrame(const FOpenXRAPIPacketBase& BasePacket)
diff --git a/Engine/Plugins/Runtime/XR/XRScribe/Source/XRScribe/Private/XRScribeEmulatedPoseManager.cpp b/Engine/Plugins/Runtime/XR/XRScribe/Source/XRScribe/Private/XRScribeEmulatedPoseManager.cpp
--- a/Engine/Plugins/Runtime/XR/XRScribe/Source/XRScribe/Private/XRScribeEmulatedPoseManager.cpp
+++ b/Engine/Plugins/Runtime/XR/XRScribe/Source/XRScribe/Private/XRScribeEmulatedPoseManager.cpp
@@ -3,6 +3,7 @@
 #include "XRScribeEmulatedPoseManager.h"
 #include "XRScribeFileFormat.h"
 
+#include "HAL/IConsoleManager.h"
 #include "Logging/LogMacros.h"
 
 DEFINE_LOG_CATEGORY_STATIC(LogXRScribePoseManager, Log, All);
@@ -10,6 +11,31 @@ DEFINE_LOG_CATEGORY_STATIC(LogXRScribePoseManager, Log, All);
 namespace UE::XRScribe
 {
 
+enum class EXRScribePoseReplayMode : int32
+{
+	/** Replay poses over the same time period they were captured */
+	TimeMatched,
+	/** Replay poses as fast as possible (might look slow or fast across different machines/build configs) */
+	Immediate,
+};
+
+ static EXRScribePoseReplayMode XRScribePoseReplayMode = EXRScribePoseReplayMode::TimeMatched;
+ static FAutoConsoleVariableRef CVarXRScribePoseReplay(TEXT("XRScribe.PoseReplayMode"),
+ 	reinterpret_cast<int32&>(XRScribePoseReplayMode),
+ 	TEXT("Toggle the pose replay mode for XRScribe (TimeMatched, Immediate)"),
+ 	ECVF_ReadOnly);
+	
+void FOpenXRPoseManager::RegisterCapturedWaitFrames(const TArray<FOpenXRWaitFramePacket>& InWaitFrameHistory)
+{
+	WaitFrameStart = InWaitFrameHistory[0].FrameState.predictedDisplayTime;
+	WaitFrameEnd = InWaitFrameHistory.Last().FrameState.predictedDisplayTime;
+
+	// generate 120Hz slices for entire capture, as we assume we won't sample poses faster than 120Hz
+	CapturedSliceCount = ((WaitFrameEnd - WaitFrameStart) / 8333333) + 1;
+
+	WaitFrameHistory = InWaitFrameHistory;
+}
+
 void FOpenXRPoseManager::RegisterCapturedReferenceSpaces(const TArray<FOpenXRCreateReferenceSpacePacket>& CreateReferenceSpacePackets)
 {
 	CapturedReferenceSpaces.Append(CreateReferenceSpacePackets);
@@ -99,6 +125,47 @@ bool FOpenXRPoseManager::FilterSpaceHistory(const TArray<FOpenXRLocateSpacePacke
 	return true;
 }
 
+TArray<FOpenXRLocateSpacePacket> FOpenXRPoseManager::SliceFilteredSpaceHistory(const TArray<FOpenXRLocateSpacePacket>& FilteredHistory)
+{
+	TArray<FOpenXRLocateSpacePacket> SlicedHistory;
+
+	if (XRScribePoseReplayMode == EXRScribePoseReplayMode::TimeMatched)
+	{
+		SlicedHistory.Reserve(CapturedSliceCount);
+
+		FOpenXRLocateSpacePacket PropogatedLocationPacket = FilteredHistory.Top();
+		PropogatedLocationPacket.Location.locationFlags = XR_SPACE_LOCATION_ORIENTATION_VALID_BIT | XR_SPACE_LOCATION_POSITION_VALID_BIT | XR_SPACE_LOCATION_ORIENTATION_TRACKED_BIT | XR_SPACE_LOCATION_POSITION_TRACKED_BIT;
+		PropogatedLocationPacket.Location.pose.orientation = ToXrQuat(FQuat::Identity);
+		PropogatedLocationPacket.Location.pose.position = ToXrVector(FVector::ZeroVector);
+
+		int64 HistoryIndex = 0;
+		for (int64 SliceIndex = 0; SliceIndex < CapturedSliceCount; SliceIndex++)
+		{
+			PropogatedLocationPacket.Time = WaitFrameStart + (SliceIndex * 8333333);
+
+			// TODO: we could interpolate poses in the slice?
+			// TODO: should we clean out invalid poses?
+
+			// forward propagate current pose until slice time surpasses current history time
+			if (HistoryIndex < FilteredHistory.Num() &&
+				FilteredHistory[HistoryIndex].Time <= PropogatedLocationPacket.Time)
+			{
+				PropogatedLocationPacket.Location = FilteredHistory[HistoryIndex].Location;
+				HistoryIndex++;
+			}
+		
+			SlicedHistory.Add(PropogatedLocationPacket);
+		}
+	}
+	else
+	{
+		check(XRScribePoseReplayMode == EXRScribePoseReplayMode::Immediate);
+		SlicedHistory = FilteredHistory;
+	}
+
+	return SlicedHistory;
+}
+
 void FOpenXRPoseManager::ProcessCapturedReferenceSpaceHistory(const TArray<FOpenXRLocateSpacePacket>& SpaceHistory)
 {
 	const XrSpace CapturedLocatedSpace = SpaceHistory[0].Space;
@@ -113,7 +180,9 @@ void FOpenXRPoseManager::ProcessCapturedReferenceSpaceHistory(const TArray<FOpen
 	TArray<FOpenXRLocateSpacePacket> FilteredSpaceHistory;
 	if (FilterSpaceHistory(SpaceHistory, FilteredSpaceHistory))
 	{
-		ReferencePoseHistories.Add(XR_REFERENCE_SPACE_TYPE_VIEW, MoveTemp(FilteredSpaceHistory));
+		TArray<FOpenXRLocateSpacePacket> SlicedHistory = SliceFilteredSpaceHistory(FilteredSpaceHistory);
+
+		ReferencePoseHistories.Add(XR_REFERENCE_SPACE_TYPE_VIEW, MoveTemp(SlicedHistory));
 	}
 }
 
@@ -136,7 +205,9 @@ void FOpenXRPoseManager::ProcessCapturedActionSpaceHistory(const TArray<FOpenXRL
 	if (FilterSpaceHistory(SpaceHistory, FilteredSpaceHistory))
 	{
 		const FName FetchedActionName(CreateAction.ActionName.GetData());
-		ActionPoseHistories.Add(FetchedActionName, MoveTemp(FilteredSpaceHistory));
+
+		TArray<FOpenXRLocateSpacePacket> SlicedHistory = SliceFilteredSpaceHistory(FilteredSpaceHistory);
+		ActionPoseHistories.Add(FetchedActionName, MoveTemp(SlicedHistory));
 
 		// TODO: we should be resilient to multiple action spaces per name
 	}
@@ -196,6 +267,8 @@ void FOpenXRPoseManager::AddEmulatedFrameTime(XrTime Time, int32 FrameNum)
 	FrameTimeHistory[HistoryIndex] = FrameTimeHistoryEntry({ Time , FrameNum });
 
 	LastInsertedFrameIndex = HistoryIndex;
+
+	EmulatedBaseTime = FMath::Min(EmulatedBaseTime, Time);
 }
 
 bool FOpenXRPoseManager::DoesActionContainPoseHistory(TStaticArray<ANSICHAR, XR_MAX_ACTION_NAME_SIZE>& ActionName)
@@ -217,28 +290,41 @@ XrSpaceLocation FOpenXRPoseManager::GetEmulatedPoseForTime(XrSpace LocatingSpace
 		return DummyLocation;
 	}
 
-	int32 SearchIndex = LastInsertedFrameIndex;
-	int32 StopIndex = SearchIndex - PoseHistorySize;
-	StopIndex = (StopIndex < 1) ? 1 : StopIndex;
-	while (SearchIndex >= StopIndex)
+	int32 PoseIndex = 0;
+	if (XRScribePoseReplayMode == EXRScribePoseReplayMode::TimeMatched)
 	{
-		const int32 HistoryIndex = SearchIndex % PoseHistorySize;
-		if (FrameTimeHistory[HistoryIndex].Time == Time)
-		{
-			break;
-		}
-
-		SearchIndex--;
+		const XrTime TimeOffsetFromSessionStart = Time - EmulatedBaseTime;
+		const int64 SliceOffset = TimeOffsetFromSessionStart/ 8333333;
+		const int64 SliceIndex = SliceOffset % CapturedSliceCount;
+		PoseIndex = SliceIndex;		
 	}
-
-	if (SearchIndex < StopIndex)
+	else
 	{
-		// We didn't find a match, just use most recent pose
-		SearchIndex = LastInsertedFrameIndex;
+		check(XRScribePoseReplayMode == EXRScribePoseReplayMode::Immediate);
+		
+		int32 SearchIndex = LastInsertedFrameIndex;
+		int32 StopIndex = SearchIndex - PoseHistorySize;
+		StopIndex = (StopIndex < 1) ? 1 : StopIndex;
+		while (SearchIndex >= StopIndex)
+		{
+			const int32 HistoryIndex = SearchIndex % PoseHistorySize;
+			if (FrameTimeHistory[HistoryIndex].Time == Time)
+			{
+				break;
+			}
+		
+			SearchIndex--;
+		}
+		
+		if (SearchIndex < StopIndex)
+		{
+			// We didn't find a match, just use most recent pose
+			SearchIndex = LastInsertedFrameIndex;
+		}
+		
+		PoseIndex = FrameTimeHistory[SearchIndex].FrameIndex;
 	}
-
-	const int32 PoseIndex = FrameTimeHistory[SearchIndex].FrameIndex;
-
+	
 	if (EmulatedReferenceSpaceTypeMap.Contains(LocatingSpace) && EmulatedReferenceSpaceTypeMap[LocatingSpace] == XR_REFERENCE_SPACE_TYPE_VIEW)
 	{
 		const TArray<FOpenXRLocateSpacePacket>& CapturedPoseHistories = ReferencePoseHistories[XR_REFERENCE_SPACE_TYPE_VIEW];
diff --git a/Engine/Plugins/Runtime/XR/XRScribe/Source/XRScribe/Private/XRScribeEmulationLayer.cpp b/Engine/Plugins/Runtime/XR/XRScribe/Source/XRScribe/Private/XRScribeEmulationLayer.cpp
--- a/Engine/Plugins/Runtime/XR/XRScribe/Source/XRScribe/Private/XRScribeEmulationLayer.cpp
+++ b/Engine/Plugins/Runtime/XR/XRScribe/Source/XRScribe/Private/XRScribeEmulationLayer.cpp
@@ -341,6 +341,8 @@ void FOpenXREmulationLayer::PostLoadActions()
 		FCStringAnsi::Strncpy(&EmulatedSystemProperties.systemName[ActualSystemNameLen], EmuStringAddend, RemainingSpace);
 	}
 
+	PoseManager.RegisterCapturedWaitFrames(CaptureDecoder.GetWaitFrames());
+
 	PoseManager.RegisterCapturedReferenceSpaces(CaptureDecoder.GetCreatedReferenceSpaces());
 
 	PoseManager.RegisterCapturedActions(CaptureDecoder.GetCreatedActions());
@@ -990,6 +992,8 @@ XrResult FOpenXREmulationLayer::XrLayerDestroySession(XrSession session)
 	// from the spec: "The application is responsible for ensuring that it has no calls using session in progress when the session is destroyed."
 	CurrentSession.Reset();
 
+	PoseManager.OnSessionTeardown();
+
 	return XR_SUCCESS;
 }
 
@@ -1933,7 +1937,8 @@ XrResult FOpenXREmulationLayer::XrLayerLocateViews(XrSession session, const XrVi
 		return XR_ERROR_HANDLE_INVALID;
 	}
 
-	// we only support locating views relative to HMD (for now)
+	// we only support locating views relative to HMD (for now). This means the views are actually
+	// relatively static unless user would change their IPD. Might change with gaze tracking
 	// TODO: support arbitrary spaces for view location (or maybe constrain to reference spaces)
 	{
 		const FOpenXREmulatedSpace* EmulatedSpace = reinterpret_cast<FOpenXREmulatedSpace*>(viewLocateInfo->space);
diff --git a/Engine/Plugins/Runtime/XR/XRScribe/Source/XRScribe/Public/XRScribeAPIDecoder.h b/Engine/Plugins/Runtime/XR/XRScribe/Source/XRScribe/Public/XRScribeAPIDecoder.h
--- a/Engine/Plugins/Runtime/XR/XRScribe/Source/XRScribe/Public/XRScribeAPIDecoder.h
+++ b/Engine/Plugins/Runtime/XR/XRScribe/Source/XRScribe/Public/XRScribeAPIDecoder.h
@@ -39,6 +39,7 @@ public:
 	[[nodiscard]] const TMap<XrReferenceSpaceType, XrExtent2Df>& GetReferenceSpaceBounds() { return ReferenceSpaceBounds; }
 	[[nodiscard]] const TArray<int64>& GetSwapchainFormats() { return SwapchainFormats; }
 	[[nodiscard]] const TArray<FOpenXRCreateActionPacket>& GetCreatedActions() { return CreatedActions; }
+	[[nodiscard]] const TArray<FOpenXRWaitFramePacket>& GetWaitFrames() { return WaitFrames; }
 
 	TArray<uint8>& GetEncodedData()
 	{
@@ -159,6 +160,8 @@ protected:
 
 	TArray<FOpenXRCreateActionPacket> CreatedActions;
 
+	TArray<FOpenXRWaitFramePacket> WaitFrames;
+
 	// TODO: Would I ever want to bin properties into per-instance collections?
 	// When we are repaying, we're just going to create our own set of 'valid' parameters
 	// that our replay runtime supports, do we really care about how the original application run
diff --git a/Engine/Plugins/Runtime/XR/XRScribe/Source/XRScribe/Public/XRScribeEmulatedPoseManager.h b/Engine/Plugins/Runtime/XR/XRScribe/Source/XRScribe/Public/XRScribeEmulatedPoseManager.h
--- a/Engine/Plugins/Runtime/XR/XRScribe/Source/XRScribe/Public/XRScribeEmulatedPoseManager.h
+++ b/Engine/Plugins/Runtime/XR/XRScribe/Source/XRScribe/Public/XRScribeEmulatedPoseManager.h
@@ -33,6 +33,15 @@ public:
 	// Number of poses the manager will keep available relative to 'current' pose, in order for multi-frame pose queries
 	static const int32 PoseHistorySize = 5;
 
+	/**
+	 * Register information about wait frames from a capture file, in order to build a history of
+	 * display times
+	 *
+	 * @param InWaitFrameHistory The WaitFrame packets from the capture
+	 *
+	 */
+	void RegisterCapturedWaitFrames(const TArray<FOpenXRWaitFramePacket>& InWaitFrameHistory);
+
 	/**
 	 * Register information about reference spaces from a capture file, in order to build a history of
 	 * space locations relative to those reference spaces.
@@ -120,6 +129,15 @@ public:
 	 */
 	XrSpaceLocation GetEmulatedPoseForTime(XrSpace LocatingSpace, XrSpace BaseSpace, XrTime Time);
 
+	/**
+	 * Reset internal emulation state based on session teardown.
+	 *
+	 */
+	void OnSessionTeardown()
+	{
+		EmulatedBaseTime = INT64_MAX;
+	}
+	
 	// TODO: add LocateViews to pose manager tracking
 
 private:
@@ -136,6 +154,11 @@ private:
 	 */
 	bool FilterSpaceHistory(const TArray<FOpenXRLocateSpacePacket>& RawHistory, TArray<FOpenXRLocateSpacePacket>& FilteredHistory);
 
+	/**
+	 * Resample the filtered history in order to ease reading poses back based on time offsets during emulation
+	 */
+	TArray<FOpenXRLocateSpacePacket> SliceFilteredSpaceHistory(const TArray<FOpenXRLocateSpacePacket>& FilteredHistory);
+
 	void ProcessCapturedReferenceSpaceHistory(const TArray<FOpenXRLocateSpacePacket>& SpaceHistory);
 	void ProcessCapturedActionSpaceHistory(const TArray<FOpenXRLocateSpacePacket>& SpaceHistory);
 
@@ -165,6 +188,13 @@ private:
 	// History of frame times, in order for pose manager to look into to estimate the relative frame index for times
 	TStaticArray <FrameTimeHistoryEntry, PoseHistorySize> FrameTimeHistory;
 	int32 LastInsertedFrameIndex = 0;
+
+	int64 CapturedSliceCount = 0;
+	XrTime WaitFrameStart = 0;
+	XrTime WaitFrameEnd = 0;
+	TArray<FOpenXRWaitFramePacket> WaitFrameHistory;
+
+	XrTime EmulatedBaseTime = INT64_MAX;
 };
 
 } // namespace UE::XRScribe
