XRScribe: extended validation of action sets and actions

#rnx

diff --git a/Engine/Plugins/Runtime/XR/XRScribe/Source/XRScribe/Private/XRScribeEmulationLayer.cpp b/Engine/Plugins/Runtime/XR/XRScribe/Source/XRScribe/Private/XRScribeEmulationLayer.cpp
--- a/Engine/Plugins/Runtime/XR/XRScribe/Source/XRScribe/Private/XRScribeEmulationLayer.cpp
+++ b/Engine/Plugins/Runtime/XR/XRScribe/Source/XRScribe/Private/XRScribeEmulationLayer.cpp
@@ -21,6 +21,14 @@ DEFINE_LOG_CATEGORY_STATIC(LogXRScribeEmulate, Log, All);
 // * useful error logs at fail points
 // * per-api next pointer validation
 
+// Action state processing
+// Get list of waitframes + sync actions
+// associate SyncAction call with waitframe
+// sort GetActionState calls by time
+// associate GetActionState with syncaction + therefore, a frame
+// Splice states into 120Hz samples
+// resample states based on timeoffset from initial waitframe
+
 namespace UE::XRScribe
 {
 
@@ -75,6 +83,19 @@ struct FAnsiStringToPathMapKeyFuncs : TDefaultMapHashableKeyFuncs<const ANSICHAR
 	}
 };
 
+struct FAnsiStringKeyFuncs : DefaultKeyFuncs<const ANSICHAR*>
+{
+	static FORCEINLINE bool Matches(KeyInitType A, KeyInitType B)
+	{
+		return FCStringAnsi::Stricmp(A, B) == 0;
+	}
+
+	static FORCEINLINE uint32 GetKeyHash(KeyInitType Key)
+	{
+		return FCrc::Strihash_DEPRECATED(Key);
+	}
+};
+
 struct FOpenXREmulatedInstance
 {
 	FOpenXREmulatedInstance()
@@ -96,6 +117,7 @@ struct FOpenXREmulatedInstance
 	TMap<const ANSICHAR*, XrPath, FDefaultSetAllocator, FAnsiStringToPathMapKeyFuncs> ANSIStringToPathMap;
 
 	TArray<TUniquePtr<FOpenXREmulatedActionSet>> ActiveActionSets;
+	TSet<const ANSICHAR*, FAnsiStringKeyFuncs> ActiveActionSetNames;
 	
 	// TODO: is a set more appropriate/performant? Will this list get so big a linear scan would actually hurt?
 	// TODO: Maybe the FOpenXREmulatedAction destructor can go into the instance and remove itself from list?
@@ -152,6 +174,8 @@ struct FOpenXREmulatedSession
 
 	// TODO: list of action sets + actions attached
 	bool bActionSetsAttached = false;
+	TSet<XrActionSet> AttachedActionSets;
+	TSet<XrAction> AttachedActions;
 
 	bool bGraphicsBindingConfigured = false;
 	FOpenXREmulatedGraphicsBinding GraphicsBinding;
@@ -185,6 +209,7 @@ struct FOpenXREmulatedActionSet
 	uint32 Priority;
 
 	TArray<TUniquePtr<FOpenXREmulatedAction>> Actions;
+	TSet<const ANSICHAR*, FAnsiStringKeyFuncs> ActionNames;
 	bool bAttached = false;
 
 	// TODO: do we need a lock for this??
@@ -2114,7 +2139,14 @@ XrResult FOpenXREmulationLayer::XrLayerCreateActionSet(XrInstance instance, cons
 		return XR_ERROR_LOCALIZED_NAME_INVALID;
 	}
 
-	// TODO: check duplicate action set names, maybe use a set of strings?
+	{
+		FReadScopeLock InstanceLock(InstanceMutex);
+		if (CurrentInstance->ActiveActionSetNames.Contains(&createInfo->actionSetName[0]))
+		{
+			return XR_ERROR_NAME_DUPLICATED;
+		}
+	}
+
 	// TODO: validate names
 
 	TUniquePtr<FOpenXREmulatedActionSet> ActionSet = MakeUnique<FOpenXREmulatedActionSet>();
@@ -2124,6 +2156,7 @@ XrResult FOpenXREmulationLayer::XrLayerCreateActionSet(XrInstance instance, cons
 
 	{
 		FWriteScopeLock InstanceLock(InstanceMutex);
+		CurrentInstance->ActiveActionSetNames.Add(ActionSet->ActionSetName.GetData());
 		CurrentInstance->ActiveActionSets.Add(MoveTemp(ActionSet));
 		*actionSet = reinterpret_cast<XrActionSet>(CurrentInstance->ActiveActionSets.Last().Get());
 	}
@@ -2201,6 +2234,14 @@ XrResult FOpenXREmulationLayer::XrLayerCreateAction(XrActionSet actionSet, const
 		return XR_ERROR_ACTIONSETS_ALREADY_ATTACHED;
 	}
 
+	{
+		FReadScopeLock InstanceLock(InstanceMutex);
+		if (EmulatedActionSet->ActionNames.Contains(&createInfo->actionName[0]))
+		{
+			return XR_ERROR_NAME_DUPLICATED;
+		}
+	}
+
 	TUniquePtr<FOpenXREmulatedAction> Action = MakeUnique<FOpenXREmulatedAction>();
 	
 	Action->OwningActionSet = EmulatedActionSet;
@@ -2210,10 +2251,11 @@ XrResult FOpenXREmulationLayer::XrLayerCreateAction(XrActionSet actionSet, const
 	Action->SubactionPaths.Append(createInfo->subactionPaths, createInfo->countSubactionPaths);
 
 	*action = reinterpret_cast<XrAction>(Action.Get());
-	EmulatedActionSet->Actions.Add(MoveTemp(Action));
 
 	{
 		FWriteScopeLock InstanceLock(InstanceMutex);
+		EmulatedActionSet->ActionNames.Add(Action->ActionName.GetData());
+		EmulatedActionSet->Actions.Add(MoveTemp(Action));
 		CurrentInstance->AllActiveActions.Add(*action);
 	}
 
@@ -2287,16 +2329,33 @@ XrResult FOpenXREmulationLayer::XrLayerAttachSessionActionSets(XrSession session
 
 	// TODO: save off list of action sets + sets in session
 
+	TSet<XrActionSet> AttachedActionSets;
+	TSet<XrAction> AttachedActions;
+
 	for (uint32 ActionSetIndex = 0; ActionSetIndex < attachInfo->countActionSets; ActionSetIndex++)
 	{
+		if (!ActionSetHandleCheck(attachInfo->actionSets[ActionSetIndex]))
+		{
+			return XR_ERROR_HANDLE_INVALID;
+		}
+
+		AttachedActionSets.Add(attachInfo->actionSets[ActionSetIndex]);
+
 		FOpenXREmulatedActionSet* EmulatedActionSet = reinterpret_cast<FOpenXREmulatedActionSet*>(attachInfo->actionSets[ActionSetIndex]);
 		check(!EmulatedActionSet->bAttached);
 		EmulatedActionSet->bAttached = true;
+
+		for (const TUniquePtr<FOpenXREmulatedAction>& Action : EmulatedActionSet->Actions)
+		{
+			AttachedActions.Add(reinterpret_cast<XrAction>(Action.Get()));
+		}
 	}
 
 	{
 		FWriteScopeLock SessionLock(SessionMutex);
 		CurrentSession->bActionSetsAttached = true;
+		CurrentSession->AttachedActionSets = MoveTemp(AttachedActionSets);
+		CurrentSession->AttachedActions = MoveTemp(AttachedActions);
 	}
 
 	return XR_SUCCESS;
@@ -2483,7 +2542,21 @@ XrResult FOpenXREmulationLayer::XrLayerSyncActions(XrSession session, const	XrAc
 		}
 	}
 
-	// TODO: validate the action sets
+	TSet<XrActionSet> AttachedActionSets;
+	{
+		FReadScopeLock SessionLock(SessionMutex);
+		AttachedActionSets = CurrentSession->AttachedActionSets;
+	}
+	for (uint32 ActionSetIndex = 0; ActionSetIndex < syncInfo->countActiveActionSets; ActionSetIndex++)
+	{
+		if (!AttachedActionSets.Contains(syncInfo->activeActionSets[ActionSetIndex].actionSet))
+		{
+			return XR_ERROR_ACTIONSET_NOT_ATTACHED;
+		}
+	}
+
+
+	// TODO: validate the action set subaction path
 
 	LastSyncTimeTicks = FDateTime::Now().GetTicks() - StartTimeTicks;
 
