XRScribe: save off captured per-space action state histories

#rnx

diff --git a/Engine/Plugins/Runtime/XR/XRScribe/Source/XRScribe/Private/XRScribeAPIDecoder.cpp b/Engine/Plugins/Runtime/XR/XRScribe/Source/XRScribe/Private/XRScribeAPIDecoder.cpp
--- a/Engine/Plugins/Runtime/XR/XRScribe/Source/XRScribe/Private/XRScribeAPIDecoder.cpp
+++ b/Engine/Plugins/Runtime/XR/XRScribe/Source/XRScribe/Private/XRScribeAPIDecoder.cpp
@@ -673,6 +673,12 @@ void FOpenXRCaptureDecoder::DecodeGetActionStateBoolean(const FOpenXRAPIPacketBa
 	FOpenXRGetActionStateBooleanPacket Data(XrResult::XR_ERROR_RUNTIME_FAILURE);
 
 	EncodedData << Data;
+
+	if (!BooleanActionStates.Contains(Data.GetInfoBoolean.action))
+	{
+		BooleanActionStates.Add(Data.GetInfoBoolean.action);
+	}
+	BooleanActionStates[Data.GetInfoBoolean.action].Add(Data);
 }
 
 void FOpenXRCaptureDecoder::DecodeGetActionStateFloat(const FOpenXRAPIPacketBase& BasePacket)
@@ -683,6 +689,12 @@ void FOpenXRCaptureDecoder::DecodeGetActionStateFloat(const FOpenXRAPIPacketBase
 	FOpenXRGetActionStateFloatPacket Data(XrResult::XR_ERROR_RUNTIME_FAILURE);
 
 	EncodedData << Data;
+
+	if (!FloatActionStates.Contains(Data.GetInfoFloat.action))
+	{
+		FloatActionStates.Add(Data.GetInfoFloat.action);
+	}
+	FloatActionStates[Data.GetInfoFloat.action].Add(Data);
 }
 
 void FOpenXRCaptureDecoder::DecodeGetActionStateVector2f(const FOpenXRAPIPacketBase& BasePacket)
@@ -693,6 +705,12 @@ void FOpenXRCaptureDecoder::DecodeGetActionStateVector2f(const FOpenXRAPIPacketB
 	FOpenXRGetActionStateVector2fPacket Data(XrResult::XR_ERROR_RUNTIME_FAILURE);
 
 	EncodedData << Data;
+
+	if (!VectorActionStates.Contains(Data.GetInfoVector2f.action))
+	{
+		VectorActionStates.Add(Data.GetInfoVector2f.action);
+	}
+	VectorActionStates[Data.GetInfoVector2f.action].Add(Data);
 }
 
 void FOpenXRCaptureDecoder::DecodeGetActionStatePose(const FOpenXRAPIPacketBase& BasePacket)
@@ -703,6 +721,12 @@ void FOpenXRCaptureDecoder::DecodeGetActionStatePose(const FOpenXRAPIPacketBase&
 	FOpenXRGetActionStatePosePacket Data(XrResult::XR_ERROR_RUNTIME_FAILURE);
 
 	EncodedData << Data;
+
+	if (!PoseActionStates.Contains(Data.GetInfoPose.action))
+	{
+		PoseActionStates.Add(Data.GetInfoPose.action);
+	}
+	PoseActionStates[Data.GetInfoPose.action].Add(Data);
 }
 
 void FOpenXRCaptureDecoder::DecodeSyncActions(const FOpenXRAPIPacketBase& BasePacket)
@@ -715,6 +739,8 @@ void FOpenXRCaptureDecoder::DecodeSyncActions(const FOpenXRAPIPacketBase& BasePa
 	FOpenXRSyncActionsPacket Data(XrResult::XR_ERROR_RUNTIME_FAILURE);
 
 	EncodedData << Data;
+
+	SyncActions.Add(Data);
 }
 
 //void FOpenXRCaptureDecoder::DecodeEnumerateBoundSourcesForAction(const FOpenXRAPIPacketBase& BasePacket)
diff --git a/Engine/Plugins/Runtime/XR/XRScribe/Source/XRScribe/Private/XRScribeEmulatedPoseManager.cpp b/Engine/Plugins/Runtime/XR/XRScribe/Source/XRScribe/Private/XRScribeEmulatedPoseManager.cpp
--- a/Engine/Plugins/Runtime/XR/XRScribe/Source/XRScribe/Private/XRScribeEmulatedPoseManager.cpp
+++ b/Engine/Plugins/Runtime/XR/XRScribe/Source/XRScribe/Private/XRScribeEmulatedPoseManager.cpp
@@ -8,6 +8,14 @@
 
 DEFINE_LOG_CATEGORY_STATIC(LogXRScribePoseManager, Log, All);
 
+// Action state processing
+// Get list of waitframes + sync actions DONE
+// associate SyncAction call with waitframe
+// sort GetActionState calls by time
+// associate GetActionState with syncaction + therefore, a frame
+// Splice states into 120Hz samples
+// resample states based on timeoffset from initial waitframe
+
 namespace UE::XRScribe
 {
 
@@ -64,10 +72,7 @@ void FOpenXRPoseManager::RegisterCapturedActions(const TArray <FOpenXRCreateActi
 {
 	for (const FOpenXRCreateActionPacket& CreateAction : CreateActionPackets)
 	{
-		if (CreateAction.ActionType == XR_ACTION_TYPE_POSE_INPUT)
-		{
-			CapturedPoseActions.Add(CreateAction.Action, CreateAction);
-		}
+		CapturedActions.Add(CreateAction.Action, CreateAction);
 	}
 }
 
@@ -197,7 +202,7 @@ void FOpenXRPoseManager::ProcessCapturedActionSpaceHistory(const TArray<FOpenXRL
 	});
 	check(MatchingInfo);
 
-	const FOpenXRCreateActionPacket& CreateAction = CapturedPoseActions[MatchingInfo->ActionSpaceCreateInfo.action];
+	const FOpenXRCreateActionPacket& CreateAction = CapturedActions[MatchingInfo->ActionSpaceCreateInfo.action];
 
 	check(CreateAction.ActionType == XR_ACTION_TYPE_POSE_INPUT);
 
@@ -243,6 +248,49 @@ void FOpenXRPoseManager::RegisterCapturedSpaceHistories(const TMap<XrSpace, TArr
 	}
 }
 
+void FOpenXRPoseManager::RegisterCapturedActionStates(const TArray<FOpenXRSyncActionsPacket>& SyncActionsPackets,
+	const TMap<XrAction, TArray<FOpenXRGetActionStateBooleanPacket>>& BooleanActionStates,
+	const TMap<XrAction, TArray<FOpenXRGetActionStateFloatPacket>>& FloatActionStates,
+	const TMap<XrAction, TArray<FOpenXRGetActionStateVector2fPacket>>& VectorActionStates,
+	const TMap<XrAction, TArray<FOpenXRGetActionStatePosePacket>>& PoseActionStates)
+{
+
+	for (const auto& BooleanActionStateList : BooleanActionStates)
+	{
+		const FOpenXRCreateActionPacket& CreateActionPacket = CapturedActions[BooleanActionStateList.Key];
+		const FName ActionName(CreateActionPacket.ActionName.GetData());
+
+		BooleanActionStateHistories.Add(ActionName, BooleanActionStateList.Value);
+	}
+
+	for (const auto& FloatActionStateList : FloatActionStates)
+	{
+		const FOpenXRCreateActionPacket& CreateActionPacket = CapturedActions[FloatActionStateList.Key];
+		const FName ActionName(CreateActionPacket.ActionName.GetData());
+
+		FloatActionStateHistories.Add(ActionName, FloatActionStateList.Value);
+	}
+
+	for (const auto& VectorActionStateList : VectorActionStates)
+	{
+		const FOpenXRCreateActionPacket& CreateActionPacket = CapturedActions[VectorActionStateList.Key];
+		const FName ActionName(CreateActionPacket.ActionName.GetData());
+
+		VectorActionStateHistories.Add(ActionName, VectorActionStateList.Value);
+	}
+
+	for (const auto& PoseActionStateList : PoseActionStates)
+	{
+		const FOpenXRCreateActionPacket& CreateActionPacket = CapturedActions[PoseActionStateList.Key];
+		const FName ActionName(CreateActionPacket.ActionName.GetData());
+
+		PoseActionStateHistories.Add(ActionName, PoseActionStateList.Value);
+	}
+
+	// map sync actions to wait frames
+
+}
+
 void FOpenXRPoseManager::RegisterEmulatedReferenceSpace(const XrReferenceSpaceCreateInfo& CreateInfo, XrSpace SpaceHandle)
 {
 	check(!EmulatedReferenceSpaceTypeMap.Contains(SpaceHandle));
diff --git a/Engine/Plugins/Runtime/XR/XRScribe/Source/XRScribe/Private/XRScribeEmulationLayer.cpp b/Engine/Plugins/Runtime/XR/XRScribe/Source/XRScribe/Private/XRScribeEmulationLayer.cpp
--- a/Engine/Plugins/Runtime/XR/XRScribe/Source/XRScribe/Private/XRScribeEmulationLayer.cpp
+++ b/Engine/Plugins/Runtime/XR/XRScribe/Source/XRScribe/Private/XRScribeEmulationLayer.cpp
@@ -21,14 +21,6 @@ DEFINE_LOG_CATEGORY_STATIC(LogXRScribeEmulate, Log, All);
 // * useful error logs at fail points
 // * per-api next pointer validation
 
-// Action state processing
-// Get list of waitframes + sync actions
-// associate SyncAction call with waitframe
-// sort GetActionState calls by time
-// associate GetActionState with syncaction + therefore, a frame
-// Splice states into 120Hz samples
-// resample states based on timeoffset from initial waitframe
-
 namespace UE::XRScribe
 {
 
@@ -374,6 +366,12 @@ void FOpenXREmulationLayer::PostLoadActions()
 	PoseManager.RegisterCapturedActionSpaces(CaptureDecoder.GetCreatedActionSpaces());
 
 	PoseManager.RegisterCapturedSpaceHistories(CaptureDecoder.GetSpaceLocations());
+
+	PoseManager.RegisterCapturedActionStates(CaptureDecoder.GetSyncActions(),
+		CaptureDecoder.GetBooleanActionStates(), 
+		CaptureDecoder.GetFloatActionStates(), 
+		CaptureDecoder.GetVectorActionStates(), 
+		CaptureDecoder.GetPoseActionStates());
 }
 
 bool FOpenXREmulationLayer::InstanceHandleCheck(XrInstance InputInstance)
diff --git a/Engine/Plugins/Runtime/XR/XRScribe/Source/XRScribe/Public/XRScribeAPIDecoder.h b/Engine/Plugins/Runtime/XR/XRScribe/Source/XRScribe/Public/XRScribeAPIDecoder.h
--- a/Engine/Plugins/Runtime/XR/XRScribe/Source/XRScribe/Public/XRScribeAPIDecoder.h
+++ b/Engine/Plugins/Runtime/XR/XRScribe/Source/XRScribe/Public/XRScribeAPIDecoder.h
@@ -40,6 +40,11 @@ public:
 	[[nodiscard]] const TArray<int64>& GetSwapchainFormats() { return SwapchainFormats; }
 	[[nodiscard]] const TArray<FOpenXRCreateActionPacket>& GetCreatedActions() { return CreatedActions; }
 	[[nodiscard]] const TArray<FOpenXRWaitFramePacket>& GetWaitFrames() { return WaitFrames; }
+	[[nodiscard]] const TArray<FOpenXRSyncActionsPacket>& GetSyncActions() { return SyncActions; }
+	[[nodiscard]] const TMap<XrAction, TArray<FOpenXRGetActionStateBooleanPacket>>& GetBooleanActionStates() { return BooleanActionStates; }
+	[[nodiscard]] const TMap<XrAction, TArray<FOpenXRGetActionStateFloatPacket>>& GetFloatActionStates() { return FloatActionStates; }
+	[[nodiscard]] const TMap<XrAction, TArray<FOpenXRGetActionStateVector2fPacket>>& GetVectorActionStates() { return VectorActionStates; }
+	[[nodiscard]] const TMap<XrAction, TArray<FOpenXRGetActionStatePosePacket>>& GetPoseActionStates() { return PoseActionStates; }
 
 	TArray<uint8>& GetEncodedData()
 	{
@@ -162,6 +167,13 @@ protected:
 
 	TArray<FOpenXRWaitFramePacket> WaitFrames;
 
+	TArray<FOpenXRSyncActionsPacket> SyncActions;
+
+	TMap<XrAction, TArray<FOpenXRGetActionStateBooleanPacket>> BooleanActionStates;
+	TMap<XrAction, TArray<FOpenXRGetActionStateFloatPacket>> FloatActionStates;
+	TMap<XrAction, TArray<FOpenXRGetActionStateVector2fPacket>> VectorActionStates;
+	TMap<XrAction, TArray<FOpenXRGetActionStatePosePacket>> PoseActionStates;
+
 	// TODO: Would I ever want to bin properties into per-instance collections?
 	// When we are repaying, we're just going to create our own set of 'valid' parameters
 	// that our replay runtime supports, do we really care about how the original application run
diff --git a/Engine/Plugins/Runtime/XR/XRScribe/Source/XRScribe/Public/XRScribeEmulatedPoseManager.h b/Engine/Plugins/Runtime/XR/XRScribe/Source/XRScribe/Public/XRScribeEmulatedPoseManager.h
--- a/Engine/Plugins/Runtime/XR/XRScribe/Source/XRScribe/Public/XRScribeEmulatedPoseManager.h
+++ b/Engine/Plugins/Runtime/XR/XRScribe/Source/XRScribe/Public/XRScribeEmulatedPoseManager.h
@@ -77,6 +77,23 @@ public:
 	 */
 	void RegisterCapturedSpaceHistories(const TMap<XrSpace, TArray<FOpenXRLocateSpacePacket>>& SpaceHistories);
 
+
+	/**
+	 * Register history of GetActionState + SyncActions calls
+	 *
+	 * @param SyncActionsPackets List of SyncActions calls
+	 * @param BooleanActionStates Per-action history of boolean action states
+	 * @param FloatActionStates Per-action history of float action states
+	 * @param VectorActionStates Per-action history of Vector2f action states
+	 * @param PoseActionStates Per-action history of Pose action states (indicating whether pose is valid or not, not actual pose data)
+	 *
+	 */
+	void RegisterCapturedActionStates(const TArray<FOpenXRSyncActionsPacket>& SyncActionsPackets, 
+		const TMap<XrAction, TArray<FOpenXRGetActionStateBooleanPacket>>& BooleanActionStates,
+		const TMap<XrAction, TArray<FOpenXRGetActionStateFloatPacket>>& FloatActionStates,
+		const TMap<XrAction, TArray<FOpenXRGetActionStateVector2fPacket>>& VectorActionStates,
+		const TMap<XrAction, TArray<FOpenXRGetActionStatePosePacket>>& PoseActionStates);
+
 	/**
 	 * Register emulated reference space. Pose manager needs this to match the new emulated reference space against
 	 * a captured reference space of similar attributes.
@@ -165,7 +182,7 @@ private:
 	// Known spaces and relevant actions from capture
 	TArray<FOpenXRCreateReferenceSpacePacket> CapturedReferenceSpaces;
 	TArray<FOpenXRCreateActionSpacePacket> CapturedActionSpaces;
-	TMap<XrAction, FOpenXRCreateActionPacket> CapturedPoseActions;
+	TMap<XrAction, FOpenXRCreateActionPacket> CapturedActions;
 
 	/**
 	 * Pose histories for space types, which could be unified across different space handles
@@ -174,6 +191,11 @@ private:
 	TMap<XrReferenceSpaceType, TArray<FOpenXRLocateSpacePacket>> ReferencePoseHistories;
 	TMap<FName, TArray<FOpenXRLocateSpacePacket>> ActionPoseHistories;
 
+	TMap<FName, TArray<FOpenXRGetActionStateBooleanPacket>> BooleanActionStateHistories;
+	TMap<FName, TArray<FOpenXRGetActionStateFloatPacket>> FloatActionStateHistories;
+	TMap<FName, TArray<FOpenXRGetActionStateVector2fPacket>> VectorActionStateHistories;
+	TMap<FName, TArray<FOpenXRGetActionStatePosePacket>> PoseActionStateHistories;
+
 	// Map of emulated spaces to their underlying reference space type (if applicable)
 	TMap<XrSpace, XrReferenceSpaceType> EmulatedReferenceSpaceTypeMap;
 
