XRScribe: input replay

* Cache off pose histories and action state fetches, and process them all at once
* Splice action state into 120Hz, but no interpolation between states
* Use ActionNames + Subpaths to match between captured data and emulation requests
* Rename PoseManager -> ActionPoseManager
* Don't re-run action + pose history, just execute one replay

diff --git a/Engine/Plugins/Runtime/XR/XRScribe/Source/XRScribe/Private/XRScribeAPIDecoder.cpp b/Engine/Plugins/Runtime/XR/XRScribe/Source/XRScribe/Private/XRScribeAPIDecoder.cpp
--- a/Engine/Plugins/Runtime/XR/XRScribe/Source/XRScribe/Private/XRScribeAPIDecoder.cpp
+++ b/Engine/Plugins/Runtime/XR/XRScribe/Source/XRScribe/Private/XRScribeAPIDecoder.cpp
@@ -287,15 +287,16 @@ void FOpenXRCaptureDecoder::DecodeCreateReferenceSpace(const FOpenXRAPIPacketBas
 void FOpenXRCaptureDecoder::DecodeGetReferenceSpaceBoundsRect(const FOpenXRAPIPacketBase& BasePacket)
 {
 	check(BasePacket.ApiId == EOpenXRAPIPacketId::GetReferenceSpaceBoundsRect);
-	check(BasePacket.Result == XR_SUCCESS);
 
 	FOpenXRGetReferenceSpaceBoundsRectPacket Data(XrResult::XR_ERROR_RUNTIME_FAILURE);
 
 	EncodedData << Data;
 
-	ReferenceSpaceBounds.Add(Data.ReferenceSpaceType, Data.Bounds);
-
-	// TODO: Check for existing bounds associated with reference space?
+	if (Data.Result == XR_SUCCESS)
+	{
+		ReferenceSpaceBounds.Add(Data.ReferenceSpaceType, Data.Bounds);
+		// TODO: Check for existing bounds associated with reference space?
+	}
 }
 
 void FOpenXRCaptureDecoder::DecodeCreateActionSpace(const FOpenXRAPIPacketBase& BasePacket)
@@ -569,7 +570,7 @@ void FOpenXRCaptureDecoder::DecodeStringToPath(const FOpenXRAPIPacketBase& BaseP
 
 	EncodedData << Data;
 
-	PathToStringMap.Add(Data.GeneratedPath, FString(ANSI_TO_TCHAR(Data.PathStringToWrite.GetData())));
+	PathToStringMap.Add(Data.GeneratedPath, FName(ANSI_TO_TCHAR(Data.PathStringToWrite.GetData())));
 
 	// TODO: Do we need a bi-directional map at any point?
 }
diff --git a/Engine/Plugins/Runtime/XR/XRScribe/Source/XRScribe/Private/XRScribeEmulatedPoseManager.cpp b/Engine/Plugins/Runtime/XR/XRScribe/Source/XRScribe/Private/XRScribeEmulatedPoseManager.cpp
--- a/Engine/Plugins/Runtime/XR/XRScribe/Source/XRScribe/Private/XRScribeEmulatedPoseManager.cpp
+++ b/Engine/Plugins/Runtime/XR/XRScribe/Source/XRScribe/Private/XRScribeEmulatedPoseManager.cpp
@@ -8,14 +8,6 @@
 
 DEFINE_LOG_CATEGORY_STATIC(LogXRScribePoseManager, Log, All);
 
-// Action state processing
-// Get list of waitframes + sync actions DONE
-// associate SyncAction call with waitframe
-// sort GetActionState calls by time
-// associate GetActionState with syncaction + therefore, a frame
-// Splice states into 120Hz samples
-// resample states based on timeoffset from initial waitframe
-
 namespace UE::XRScribe
 {
 
@@ -33,25 +25,24 @@ enum class EXRScribePoseReplayMode : int32
  	TEXT("Toggle the pose replay mode for XRScribe (TimeMatched, Immediate)"),
  	ECVF_ReadOnly);
 	
-void FOpenXRPoseManager::RegisterCapturedWaitFrames(const TArray<FOpenXRWaitFramePacket>& InWaitFrameHistory)
-{
-	WaitFrameStart = InWaitFrameHistory[0].FrameState.predictedDisplayTime;
-	WaitFrameEnd = InWaitFrameHistory.Last().FrameState.predictedDisplayTime;
-
-	// generate 120Hz slices for entire capture, as we assume we won't sample poses faster than 120Hz
-	CapturedSliceCount = ((WaitFrameEnd - WaitFrameStart) / 8333333) + 1;
+ void FOpenXRActionPoseManager::RegisterCapturedPathStrings(const TMap<XrPath, FName>& PathStringMap)
+ {
+	 CapturedPaths = PathStringMap;
+ }
 
+void FOpenXRActionPoseManager::RegisterCapturedWaitFrames(const TArray<FOpenXRWaitFramePacket>& InWaitFrameHistory)
+{
 	WaitFrameHistory = InWaitFrameHistory;
 }
 
-void FOpenXRPoseManager::RegisterCapturedReferenceSpaces(const TArray<FOpenXRCreateReferenceSpacePacket>& CreateReferenceSpacePackets)
+void FOpenXRActionPoseManager::RegisterCapturedReferenceSpaces(const TArray<FOpenXRCreateReferenceSpacePacket>& CreateReferenceSpacePackets)
 {
 	CapturedReferenceSpaces.Append(CreateReferenceSpacePackets);
 
 	// TODO: check for doubles?
 }
 
-XrReferenceSpaceType FOpenXRPoseManager::GetCapturedReferenceSpaceType(XrSpace CapturedSpace)
+XrReferenceSpaceType FOpenXRActionPoseManager::GetCapturedReferenceSpaceType(XrSpace CapturedSpace)
 {
 	XrReferenceSpaceType Result = XR_REFERENCE_SPACE_TYPE_MAX_ENUM;
 
@@ -68,7 +59,7 @@ XrReferenceSpaceType FOpenXRPoseManager::GetCapturedReferenceSpaceType(XrSpace C
 	return Result;
 }
 
-void FOpenXRPoseManager::RegisterCapturedActions(const TArray <FOpenXRCreateActionPacket>& CreateActionPackets)
+void FOpenXRActionPoseManager::RegisterCapturedActions(const TArray <FOpenXRCreateActionPacket>& CreateActionPackets)
 {
 	for (const FOpenXRCreateActionPacket& CreateAction : CreateActionPackets)
 	{
@@ -76,12 +67,12 @@ void FOpenXRPoseManager::RegisterCapturedActions(const TArray <FOpenXRCreateActi
 	}
 }
 
-void FOpenXRPoseManager::RegisterCapturedActionSpaces(const TArray<FOpenXRCreateActionSpacePacket>& CreateActionSpacePackets)
+void FOpenXRActionPoseManager::RegisterCapturedActionSpaces(const TArray<FOpenXRCreateActionSpacePacket>& CreateActionSpacePackets)
 {
 	CapturedActionSpaces.Append(CreateActionSpacePackets);
 }
 
-bool FOpenXRPoseManager::VerifyCapturedActionSpace(XrSpace CapturedSpace)
+bool FOpenXRActionPoseManager::VerifyCapturedActionSpace(XrSpace CapturedSpace)
 {
 	return CapturedActionSpaces.ContainsByPredicate([CapturedSpace](const FOpenXRCreateActionSpacePacket& CreateActionSpaceInfo)
 	{
@@ -89,7 +80,7 @@ bool FOpenXRPoseManager::VerifyCapturedActionSpace(XrSpace CapturedSpace)
 	});
 }
 
-bool FOpenXRPoseManager::FilterSpaceHistory(const TArray<FOpenXRLocateSpacePacket>& RawHistory, TArray<FOpenXRLocateSpacePacket>& FilteredHistory)
+bool FOpenXRActionPoseManager::FilterSpaceHistory(const TArray<FOpenXRLocateSpacePacket>& RawHistory, TArray<FOpenXRLocateSpacePacket>& FilteredHistory)
 {
 	XrSpace CapturedBaseSpace = RawHistory[0].BaseSpace;
 	XrReferenceSpaceType BaseRefSpaceType = GetCapturedReferenceSpaceType(CapturedBaseSpace);
@@ -130,7 +121,7 @@ bool FOpenXRPoseManager::FilterSpaceHistory(const TArray<FOpenXRLocateSpacePacke
 	return true;
 }
 
-TArray<FOpenXRLocateSpacePacket> FOpenXRPoseManager::SliceFilteredSpaceHistory(const TArray<FOpenXRLocateSpacePacket>& FilteredHistory)
+TArray<FOpenXRLocateSpacePacket> FOpenXRActionPoseManager::SliceFilteredSpaceHistory(const TArray<FOpenXRLocateSpacePacket>& FilteredHistory)
 {
 	TArray<FOpenXRLocateSpacePacket> SlicedHistory;
 
@@ -146,7 +137,7 @@ TArray<FOpenXRLocateSpacePacket> FOpenXRPoseManager::SliceFilteredSpaceHistory(c
 		int64 HistoryIndex = 0;
 		for (int64 SliceIndex = 0; SliceIndex < CapturedSliceCount; SliceIndex++)
 		{
-			PropogatedLocationPacket.Time = WaitFrameStart + (SliceIndex * 8333333);
+			PropogatedLocationPacket.Time = CapturedRangeStart + (SliceIndex * 8333333);
 
 			// TODO: we could interpolate poses in the slice?
 			// TODO: should we clean out invalid poses?
@@ -171,7 +162,7 @@ TArray<FOpenXRLocateSpacePacket> FOpenXRPoseManager::SliceFilteredSpaceHistory(c
 	return SlicedHistory;
 }
 
-void FOpenXRPoseManager::ProcessCapturedReferenceSpaceHistory(const TArray<FOpenXRLocateSpacePacket>& SpaceHistory)
+void FOpenXRActionPoseManager::ProcessCapturedReferenceSpaceHistory(const TArray<FOpenXRLocateSpacePacket>& SpaceHistory)
 {
 	const XrSpace CapturedLocatedSpace = SpaceHistory[0].Space;
 	const XrReferenceSpaceType RefSpaceType = GetCapturedReferenceSpaceType(CapturedLocatedSpace);
@@ -191,7 +182,7 @@ void FOpenXRPoseManager::ProcessCapturedReferenceSpaceHistory(const TArray<FOpen
 	}
 }
 
-void FOpenXRPoseManager::ProcessCapturedActionSpaceHistory(const TArray<FOpenXRLocateSpacePacket>& SpaceHistory)
+void FOpenXRActionPoseManager::ProcessCapturedActionSpaceHistory(const TArray<FOpenXRLocateSpacePacket>& SpaceHistory)
 {
 	const XrSpace LocatedSpace = SpaceHistory[0].Space;
 
@@ -218,9 +209,91 @@ void FOpenXRPoseManager::ProcessCapturedActionSpaceHistory(const TArray<FOpenXRL
 	}
 }
 
-void FOpenXRPoseManager::RegisterCapturedSpaceHistories(const TMap<XrSpace, TArray<FOpenXRLocateSpacePacket>>& SpaceHistories)
+void FOpenXRActionPoseManager::RegisterCapturedSpaceHistories(const TMap<XrSpace, TArray<FOpenXRLocateSpacePacket>>& SpaceHistories)
+{
+	CapturedRawSpaceHistories = SpaceHistories;
+}
+
+void FOpenXRActionPoseManager::RegisterCapturedActionStates(const TArray<FOpenXRSyncActionsPacket>& SyncActionsPackets,
+	const TMap<XrAction, TArray<FOpenXRGetActionStateBooleanPacket>>& BooleanActionStates,
+	const TMap<XrAction, TArray<FOpenXRGetActionStateFloatPacket>>& FloatActionStates,
+	const TMap<XrAction, TArray<FOpenXRGetActionStateVector2fPacket>>& VectorActionStates,
+	const TMap<XrAction, TArray<FOpenXRGetActionStatePosePacket>>& PoseActionStates)
+{
+
+	//CapturedRawSyncActionsPackets = SyncActionsPackets;
+	CapturedRawBooleanActionStates = BooleanActionStates;
+	CapturedRawFloatActionStates = FloatActionStates;
+	CapturedRawVectorActionStates = VectorActionStates;
+	CapturedRawPoseActionStates = PoseActionStates;
+}
+
+void FOpenXRActionPoseManager::CalculateCapturedTimeRange()
+{
+	// TODO: We might want to offset this time 'back' based on the predictedDisplayPeriod. That would estimate the CPU time start for the session
+	// TODO: We might want to calculate a 'real' last time based on the latest time returned from an ActionState lastChangeTime
+	XrTime WaitFrameStart = WaitFrameHistory[0].FrameState.predictedDisplayTime;
+	XrTime WaitFrameEnd = WaitFrameHistory.Last().FrameState.predictedDisplayTime;
+
+	XrTime EarliestActionTime = INT64_MAX;
+	XrTime LatestActionTime = 0;
+
+	for (const auto& ActionStateHistory : CapturedRawBooleanActionStates)
+	{
+		uint32 NumActive = 0;
+		for (const FOpenXRGetActionStateBooleanPacket& BooleanActionState : ActionStateHistory.Value)
+		{
+			if (BooleanActionState.BooleanState.isActive)
+			{
+				EarliestActionTime = FMath::Min(EarliestActionTime, BooleanActionState.BooleanState.lastChangeTime);
+				LatestActionTime = FMath::Max(LatestActionTime, BooleanActionState.BooleanState.lastChangeTime);
+				NumActive++;
+			}
+		}
+		UE_LOG(LogXRScribePoseManager, Verbose, TEXT("Bool action %p - active count: %d / %d"), ActionStateHistory.Key, NumActive, ActionStateHistory.Value.Num());
+	}
+	for (const auto& ActionStateHistory : CapturedRawFloatActionStates)
+	{
+		uint32 NumActive = 0;
+		for (const FOpenXRGetActionStateFloatPacket& FloatActionState : ActionStateHistory.Value)
+		{
+			if (FloatActionState.FloatState.isActive)
+			{
+				EarliestActionTime = FMath::Min(EarliestActionTime, FloatActionState.FloatState.lastChangeTime);
+				LatestActionTime = FMath::Max(LatestActionTime, FloatActionState.FloatState.lastChangeTime);
+				NumActive++;
+			}
+		}
+		UE_LOG(LogXRScribePoseManager, Verbose, TEXT("Float action %p - active count: %d / %d"), ActionStateHistory.Key, NumActive, ActionStateHistory.Value.Num());
+	}
+	for (const auto& ActionStateHistory : CapturedRawVectorActionStates)
+	{
+		uint32 NumActive = 0;
+		for (const FOpenXRGetActionStateVector2fPacket& VectorActionState : ActionStateHistory.Value)
+		{
+			if (VectorActionState.Vector2fState.isActive)
+			{
+				EarliestActionTime = FMath::Min(EarliestActionTime, VectorActionState.Vector2fState.lastChangeTime);
+				LatestActionTime = FMath::Max(LatestActionTime, VectorActionState.Vector2fState.lastChangeTime);
+				NumActive++;
+			}
+		}
+		UE_LOG(LogXRScribePoseManager, Verbose, TEXT("Vector2F action %p - active count: %d / %d"), ActionStateHistory.Key, NumActive, ActionStateHistory.Value.Num());
+	}
+
+	CapturedRangeStart = FMath::Min(EarliestActionTime, WaitFrameStart);
+	CapturedRangeEnd = FMath::Max(LatestActionTime, WaitFrameEnd);
+}
+
+void FOpenXRActionPoseManager::ProcessCapturedHistories()
 {
-	for (const auto& SpaceHistoryTuple : SpaceHistories)
+	CalculateCapturedTimeRange();
+
+	// generate 120Hz slices for entire capture, as we assume we won't sample poses faster than 120Hz
+	// TODO: 240 Hz might work better across disparate time domains (WaitFrame, SyncActions, GetActionState, LocateViews, LocateSpace)
+	CapturedSliceCount = ((CapturedRangeEnd - CapturedRangeStart) / 8333333) + 1;
+
+	for (const auto& SpaceHistoryTuple : CapturedRawSpaceHistories)
 	{
 		const TArray<FOpenXRLocateSpacePacket>& SpaceHistory = SpaceHistoryTuple.Value;
 		if (SpaceHistory.Num() == 0)
@@ -244,60 +317,440 @@ void FOpenXRPoseManager::RegisterCapturedSpaceHistories(const TMap<XrSpace, TArr
 		{
 			UE_LOG(LogXRScribePoseManager, Warning, TEXT("Unable to process space history for captured space: %p"), reinterpret_cast<const void*>(CapturedLocatedSpace));
 		}
-
 	}
-}
 
-void FOpenXRPoseManager::RegisterCapturedActionStates(const TArray<FOpenXRSyncActionsPacket>& SyncActionsPackets,
-	const TMap<XrAction, TArray<FOpenXRGetActionStateBooleanPacket>>& BooleanActionStates,
-	const TMap<XrAction, TArray<FOpenXRGetActionStateFloatPacket>>& FloatActionStates,
-	const TMap<XrAction, TArray<FOpenXRGetActionStateVector2fPacket>>& VectorActionStates,
-	const TMap<XrAction, TArray<FOpenXRGetActionStatePosePacket>>& PoseActionStates)
-{
-
-	for (const auto& BooleanActionStateList : BooleanActionStates)
+	// Converting raw action state captures into a processed state history
+	// * Extract valid subpaths from captured state
+	// * Filter out 'inactive' state values
+	//   * We might want to support intermediate inactive periods in the future
+	//	   Currently, we do not support 'holes' in the state history.
+	// * Compress list of active states to reflect changes only (based on change time)
+	// * Based on initial 'active' time, prepend history with inactive states to match up with full capture timespan
+	// * Slice state changes into 120 Hz history of states
+	// * add in last known state to complete unknown history
+
+	for (const auto& BooleanActionStateList : CapturedRawBooleanActionStates)
 	{
+		FBooleanActionStateHistory ProcessedBooleanStateHistory;
+		ProcessedBooleanStateHistory.StateValues.Reserve(CapturedSliceCount);
+
+		// Filter list into 'Active' entries
+		// TODO: This might not be entirely accurate. We won't capture 'holes' in the history
+		// which could happen if an action is out of scope for any reason. But for our initial
+		// implementation, it's fine.
+
+		TArray<FOpenXRGetActionStateBooleanPacket> ActiveBooleanStateList;
+		ActiveBooleanStateList.Reserve(BooleanActionStateList.Value.Num());
+
+		// TODO: Scan for any valid subpaths
+		for (const FOpenXRGetActionStateBooleanPacket& Packet : BooleanActionStateList.Value)
+		{
+			if (Packet.BooleanState.isActive == XR_TRUE)
+			{
+				ActiveBooleanStateList.Add(Packet);
+				if (Packet.GetInfoBoolean.subactionPath != XR_NULL_PATH)
+				{
+					ProcessedBooleanStateHistory.ValidSubpaths.Add(CapturedPaths[Packet.GetInfoBoolean.subactionPath]);
+				}
+			}
+		}
+
+		if (ActiveBooleanStateList.IsEmpty())
+		{
+			continue;
+		}
+
+		// Compress the active list to remove redundant entries
+		TArray<FOpenXRGetActionStateBooleanPacket> CompressedActiveBooleanStateList;
+		CompressedActiveBooleanStateList.Reserve(ActiveBooleanStateList.Num());
+		CompressedActiveBooleanStateList.Add(ActiveBooleanStateList[0]);
+		for (const FOpenXRGetActionStateBooleanPacket& Packet : ActiveBooleanStateList)
+		{
+			if (Packet.BooleanState.lastChangeTime != CompressedActiveBooleanStateList.Last().BooleanState.lastChangeTime)
+			{
+				CompressedActiveBooleanStateList.Add(Packet);
+			}
+		}
+
+		// TODO: Prefill any slices with 'inactive'
+		int64 SliceIndex = 0;
+
+		while (SliceIndex < CapturedSliceCount)
+		{
+			const XrTime SliceTime = CapturedRangeStart + (SliceIndex * 8333333);
+
+			if (SliceTime >= CompressedActiveBooleanStateList[0].BooleanState.lastChangeTime)
+			{
+				break;
+			}
+
+			ProcessedBooleanStateHistory.StateValues.Add(FBooleanActionState());
+
+			SliceIndex++;
+		}
+
+		// Double walk lists to insert state changes
+		int64 HistoryIndex = 0;
+
+		while ((SliceIndex < CapturedSliceCount) &&
+			   (HistoryIndex < CompressedActiveBooleanStateList.Num()))
+		{
+			const XrTime SliceTime = CapturedRangeStart + (SliceIndex * 8333333);
+
+			if (SliceTime >= CompressedActiveBooleanStateList[HistoryIndex].BooleanState.lastChangeTime)
+			{
+				FBooleanActionState StateEntry;
+				StateEntry.Active = XR_TRUE;
+				StateEntry.State = CompressedActiveBooleanStateList[HistoryIndex].BooleanState.currentState;
+				ProcessedBooleanStateHistory.StateValues.Add(StateEntry);
+
+				HistoryIndex++;
+			}
+			else
+			{
+				// Forward propogate state from last slice
+				// TODO: is it possible to ever forward propogate for index 0??
+				FBooleanActionState PrevState = ProcessedBooleanStateHistory.StateValues.Last();
+				ProcessedBooleanStateHistory.StateValues.Add(PrevState);
+			}
+
+			SliceIndex++;
+		}
+
+		// TODO: Fill in remainder with last known state
+		while (SliceIndex < CapturedSliceCount)
+		{
+			FBooleanActionState PrevState = ProcessedBooleanStateHistory.StateValues.Last();
+			ProcessedBooleanStateHistory.StateValues.Add(PrevState);
+			SliceIndex++;
+		}
+
+		// TODO: support arbitrary internal ranges of inactive
+		// Will need to match SyncAction calls to WaitFrames in order to estimate inactive
+		// time ranges
+
+		// add processed list to the map with the name
 		const FOpenXRCreateActionPacket& CreateActionPacket = CapturedActions[BooleanActionStateList.Key];
 		const FName ActionName(CreateActionPacket.ActionName.GetData());
+		BooleanProcessedHistories.Add(ActionName, ProcessedBooleanStateHistory);
+	}	
 
-		BooleanActionStateHistories.Add(ActionName, BooleanActionStateList.Value);
-	}
-
-	for (const auto& FloatActionStateList : FloatActionStates)
+	for (const auto& FloatActionStateList : CapturedRawFloatActionStates)
 	{
+		FFloatActionStateHistory ProcessedFloatStateHistory;
+		ProcessedFloatStateHistory.StateValues.Reserve(CapturedSliceCount);
+
+		TArray<FOpenXRGetActionStateFloatPacket> ActiveFloatStateList;
+		ActiveFloatStateList.Reserve(FloatActionStateList.Value.Num());
+
+		for (const FOpenXRGetActionStateFloatPacket& Packet : FloatActionStateList.Value)
+		{
+			if (Packet.FloatState.isActive == XR_TRUE)
+			{
+				ActiveFloatStateList.Add(Packet);
+				if (Packet.GetInfoFloat.subactionPath != XR_NULL_PATH)
+				{
+					ProcessedFloatStateHistory.ValidSubpaths.Add(CapturedPaths[Packet.GetInfoFloat.subactionPath]);
+				}
+			}
+		}
+
+		if (ActiveFloatStateList.IsEmpty())
+		{
+			continue;
+		}
+
+		TArray<FOpenXRGetActionStateFloatPacket> CompressedActiveFloatStateList;
+		CompressedActiveFloatStateList.Reserve(ActiveFloatStateList.Num());
+		CompressedActiveFloatStateList.Add(ActiveFloatStateList[0]);
+		for (const FOpenXRGetActionStateFloatPacket& Packet : ActiveFloatStateList)
+		{
+			if (Packet.FloatState.lastChangeTime != CompressedActiveFloatStateList.Last().FloatState.lastChangeTime)
+			{
+				CompressedActiveFloatStateList.Add(Packet);
+			}
+		}
+
+		int64 SliceIndex = 0;
+
+		while (SliceIndex < CapturedSliceCount)
+		{
+			const XrTime SliceTime = CapturedRangeStart + (SliceIndex * 8333333);
+			if (SliceTime >= CompressedActiveFloatStateList[0].FloatState.lastChangeTime)
+			{
+				break;
+			}
+			ProcessedFloatStateHistory.StateValues.Add(FFloatActionState());
+			SliceIndex++;
+		}
+
+		int64 HistoryIndex = 0;
+
+		while ((SliceIndex < CapturedSliceCount) &&
+			(HistoryIndex < CompressedActiveFloatStateList.Num()))
+		{
+			FFloatActionState StateEntry;
+			const XrTime SliceTime = CapturedRangeStart + (SliceIndex * 8333333);
+			if (SliceTime >= CompressedActiveFloatStateList[HistoryIndex].FloatState.lastChangeTime)
+			{
+				StateEntry.Active = XR_TRUE;
+				StateEntry.State = CompressedActiveFloatStateList[HistoryIndex].FloatState.currentState;
+				HistoryIndex++;
+			}
+			else
+			{
+				// Forward propogate state from last slice
+				// TODO: is it possible to ever forward propogate for index 0??
+				StateEntry = ProcessedFloatStateHistory.StateValues.Last();
+			}
+			ProcessedFloatStateHistory.StateValues.Add(StateEntry);
+			SliceIndex++;
+		}
+
+		while (SliceIndex < CapturedSliceCount)
+		{
+			FFloatActionState PrevState = ProcessedFloatStateHistory.StateValues.Last();
+			ProcessedFloatStateHistory.StateValues.Add(PrevState);
+			SliceIndex++;
+		}
+
+		// TODO: support arbitrary internal ranges of inactive
+		// Will need to match SyncAction calls to WaitFrames in order to estimate inactive
+		// time ranges
+
 		const FOpenXRCreateActionPacket& CreateActionPacket = CapturedActions[FloatActionStateList.Key];
 		const FName ActionName(CreateActionPacket.ActionName.GetData());
-
-		FloatActionStateHistories.Add(ActionName, FloatActionStateList.Value);
+		FloatProcessedHistories.Add(ActionName, ProcessedFloatStateHistory);
 	}
 
-	for (const auto& VectorActionStateList : VectorActionStates)
+	for (const auto& VectorActionStateList : CapturedRawVectorActionStates)
 	{
+		FVector2fActionStateHistory ProcessedVector2fStateHistory;
+		ProcessedVector2fStateHistory.StateValues.Reserve(CapturedSliceCount);
+
+		TArray<FOpenXRGetActionStateVector2fPacket> ActiveVector2fStateList;
+		ActiveVector2fStateList.Reserve(VectorActionStateList.Value.Num());
+
+		for (const FOpenXRGetActionStateVector2fPacket& Packet : VectorActionStateList.Value)
+		{
+			if (Packet.Vector2fState.isActive == XR_TRUE)
+			{
+				ActiveVector2fStateList.Add(Packet);
+				if (Packet.GetInfoVector2f.subactionPath != XR_NULL_PATH)
+				{
+					ProcessedVector2fStateHistory.ValidSubpaths.Add(CapturedPaths[Packet.GetInfoVector2f.subactionPath]);
+				}
+			}
+		}
+
+		if (ActiveVector2fStateList.IsEmpty())
+		{
+			continue;
+		}
+
+		TArray<FOpenXRGetActionStateVector2fPacket> CompressedActiveVector2fStateList;
+		CompressedActiveVector2fStateList.Reserve(ActiveVector2fStateList.Num());
+		CompressedActiveVector2fStateList.Add(ActiveVector2fStateList[0]);
+		for (const FOpenXRGetActionStateVector2fPacket& Packet : ActiveVector2fStateList)
+		{
+			if (Packet.Vector2fState.lastChangeTime != CompressedActiveVector2fStateList.Last().Vector2fState.lastChangeTime)
+			{
+				CompressedActiveVector2fStateList.Add(Packet);
+			}
+		}
+
+		int64 SliceIndex = 0;
+
+		while (SliceIndex < CapturedSliceCount)
+		{
+			const XrTime SliceTime = CapturedRangeStart + (SliceIndex * 8333333);
+			if (SliceTime >= CompressedActiveVector2fStateList[0].Vector2fState.lastChangeTime)
+			{
+				break;
+			}
+			ProcessedVector2fStateHistory.StateValues.Add(FVector2fActionState());
+			SliceIndex++;
+		}
+
+		int64 HistoryIndex = 0;
+
+		while ((SliceIndex < CapturedSliceCount) &&
+			(HistoryIndex < CompressedActiveVector2fStateList.Num()))
+		{
+			FVector2fActionState StateEntry;
+			const XrTime SliceTime = CapturedRangeStart + (SliceIndex * 8333333);
+			if (SliceTime >= CompressedActiveVector2fStateList[HistoryIndex].Vector2fState.lastChangeTime)
+			{
+				StateEntry.Active = XR_TRUE;
+				StateEntry.State = CompressedActiveVector2fStateList[HistoryIndex].Vector2fState.currentState;
+				HistoryIndex++;
+			}
+			else
+			{
+				// Forward propogate state from last slice
+				// TODO: is it possible to ever forward propogate for index 0??
+				StateEntry = ProcessedVector2fStateHistory.StateValues.Last();
+			}
+			ProcessedVector2fStateHistory.StateValues.Add(StateEntry);
+			SliceIndex++;
+		}
+
+		while (SliceIndex < CapturedSliceCount)
+		{
+			FVector2fActionState PrevState = ProcessedVector2fStateHistory.StateValues.Last();
+			ProcessedVector2fStateHistory.StateValues.Add(PrevState);
+			SliceIndex++;
+		}
+
+		// TODO: support arbitrary internal ranges of inactive
+		// Will need to match SyncAction calls to WaitFrames in order to estimate inactive
+		// time ranges
+
 		const FOpenXRCreateActionPacket& CreateActionPacket = CapturedActions[VectorActionStateList.Key];
 		const FName ActionName(CreateActionPacket.ActionName.GetData());
+		Vector2fProcessedHistories.Add(ActionName, ProcessedVector2fStateHistory);
+	}
+}
+
+void FOpenXRActionPoseManager::RegisterEmulatedPath(FName PathString, XrPath Path)
+{
+	EmulatedPaths.Add(Path, PathString);
+}
 
-		VectorActionStateHistories.Add(ActionName, VectorActionStateList.Value);
+void FOpenXRActionPoseManager::RegisterEmulatedAction(FName ActionName, XrAction Action, XrActionType ActionType)
+{
+	EmulatedActionNames.Add(Action, ActionName);
+
+	if (ActionType == XR_ACTION_TYPE_BOOLEAN_INPUT)
+	{
+		XrActionStateBoolean InitialEmulatedState{};
+		InitialEmulatedState.type = XR_TYPE_ACTION_STATE_BOOLEAN;
+		InitialEmulatedState.currentState = XR_FALSE;
+		InitialEmulatedState.changedSinceLastSync = XR_FALSE;
+		InitialEmulatedState.lastChangeTime = INT64_MAX;
+		InitialEmulatedState.isActive = XR_FALSE;
+
+		SetCachedEmulatedBooleanState(ActionName, InitialEmulatedState);
+	}
+	else if (ActionType == XR_ACTION_TYPE_FLOAT_INPUT)
+	{
+		XrActionStateFloat InitialEmulatedState{};
+		InitialEmulatedState.type = XR_TYPE_ACTION_STATE_FLOAT;
+		InitialEmulatedState.currentState = 0.0f;
+		InitialEmulatedState.changedSinceLastSync = XR_FALSE;
+		InitialEmulatedState.lastChangeTime = INT64_MAX;
+		InitialEmulatedState.isActive = XR_FALSE;
+
+		SetCachedEmulatedFloatState(ActionName, InitialEmulatedState);
+	}
+	else if (ActionType == XR_ACTION_TYPE_VECTOR2F_INPUT)
+	{
+		XrActionStateVector2f InitialEmulatedState{};
+		InitialEmulatedState.type = XR_TYPE_ACTION_STATE_VECTOR2F;
+		InitialEmulatedState.currentState = { 0.0f, 0.0f };
+		InitialEmulatedState.changedSinceLastSync = XR_FALSE;
+		InitialEmulatedState.lastChangeTime = INT64_MAX;
+		InitialEmulatedState.isActive = XR_FALSE;
+
+		SetCachedEmulatedVector2fState(ActionName, InitialEmulatedState);
 	}
+	else
+	{
+		// Pose states? What else?
+	}
+}
 
-	for (const auto& PoseActionStateList : PoseActionStates)
+FName FOpenXRActionPoseManager::GetEmulatedActionName(XrAction Action)
+{
+	if (EmulatedActionNames.Contains(Action))
 	{
-		const FOpenXRCreateActionPacket& CreateActionPacket = CapturedActions[PoseActionStateList.Key];
-		const FName ActionName(CreateActionPacket.ActionName.GetData());
+		return EmulatedActionNames[Action];
+	}
+	return FName();
+}
+
+int64 FOpenXRActionPoseManager::GenerateReplaySliceIndexFromTime(XrTime EmulatedTime)
+{
+	const XrTime TimeOffsetFromSessionStart = EmulatedTime - EmulatedBaseTime;
+	const int64 SliceOffset = TimeOffsetFromSessionStart / 8333333;
+
+	// We run through the replay once, and retain the last pose/state.
+	// We could also loop, but we'd need 'reset to origin' functionality to prevent unexpected loop interactions.
+	return FMath::Min(SliceOffset, CapturedSliceCount - 1);
+}
+
+bool FOpenXRActionPoseManager::ValidateSubpath(const TSet<FName>& ValidSubpaths, XrPath ActionSubPath)
+{
+	bool bSubpathValid = true;
+	if (ActionSubPath != XR_NULL_PATH)
+	{
+		if (EmulatedPaths.Contains(ActionSubPath))
+		{
+			const FName SubactionPathName = EmulatedPaths[ActionSubPath];
+			bSubpathValid = ValidSubpaths.Contains(SubactionPathName);
+		}
+		else
+		{
+			bSubpathValid = false;
+			UE_LOG(LogXRScribePoseManager, Warning, TEXT("Unknown emulated path (%d) to check against action subpaths"), ActionSubPath);
+		}
+	}
+	return bSubpathValid;
+}
+
+const XrActionStateBoolean FOpenXRActionPoseManager::GetCachedEmulatedBooleanState(FName ActionName)
+{
+	if (CachedEmulatedBooleanStates.Contains(ActionName))
+	{
+		return CachedEmulatedBooleanStates[ActionName];
+	}
+	// TODO: check/ensure?
+	return XrActionStateBoolean({});
+}
 
-		PoseActionStateHistories.Add(ActionName, PoseActionStateList.Value);
+void FOpenXRActionPoseManager::SetCachedEmulatedBooleanState(FName ActionName, const XrActionStateBoolean& EmulatedState)
+{
+	CachedEmulatedBooleanStates.Add(ActionName, EmulatedState);
+}
+
+const XrActionStateFloat FOpenXRActionPoseManager::GetCachedEmulatedFloatState(FName ActionName)
+{
+	if (CachedEmulatedFloatStates.Contains(ActionName))
+	{
+		return CachedEmulatedFloatStates[ActionName];
 	}
+	// TODO: check/ensure?
+	return XrActionStateFloat({});
+}
+
+void FOpenXRActionPoseManager::SetCachedEmulatedFloatState(FName ActionName, const XrActionStateFloat& EmulatedState)
+{
+	CachedEmulatedFloatStates.Add(ActionName, EmulatedState);
+}
 
-	// map sync actions to wait frames
+const XrActionStateVector2f FOpenXRActionPoseManager::GetCachedEmulatedVector2fState(FName ActionName)
+{
+	if (CachedEmulatedVector2fStates.Contains(ActionName))
+	{
+		return CachedEmulatedVector2fStates[ActionName];
+	}
+	// TODO: check/ensure?
+	return XrActionStateVector2f({});
+}
 
+void FOpenXRActionPoseManager::SetCachedEmulatedVector2fState(FName ActionName, const XrActionStateVector2f& EmulatedState)
+{
+	CachedEmulatedVector2fStates.Add(ActionName, EmulatedState);
 }
 
-void FOpenXRPoseManager::RegisterEmulatedReferenceSpace(const XrReferenceSpaceCreateInfo& CreateInfo, XrSpace SpaceHandle)
+
+void FOpenXRActionPoseManager::RegisterEmulatedReferenceSpace(const XrReferenceSpaceCreateInfo& CreateInfo, XrSpace SpaceHandle)
 {
 	check(!EmulatedReferenceSpaceTypeMap.Contains(SpaceHandle));
 	EmulatedReferenceSpaceTypeMap.Add(SpaceHandle, CreateInfo.referenceSpaceType);
 }
 
-void FOpenXRPoseManager::RegisterEmulatedActionSpace(TStaticArray<ANSICHAR, XR_MAX_ACTION_NAME_SIZE>& ActionName, const XrActionSpaceCreateInfo& SpaceCreateInfo, XrSpace SpaceHandle)
+void FOpenXRActionPoseManager::RegisterEmulatedActionSpace(TStaticArray<ANSICHAR, XR_MAX_ACTION_NAME_SIZE>& ActionName, const XrActionSpaceCreateInfo& SpaceCreateInfo, XrSpace SpaceHandle)
 {
 	check(!EmulatedActionSpaceNameMap.Contains(SpaceHandle));
 
@@ -305,7 +758,7 @@ void FOpenXRPoseManager::RegisterEmulatedActionSpace(TStaticArray<ANSICHAR, XR_M
 	EmulatedActionSpaceNameMap.Add(SpaceHandle, EmulatedActionName);
 }
 
-void FOpenXRPoseManager::AddEmulatedFrameTime(XrTime Time, int32 FrameNum)
+void FOpenXRActionPoseManager::AddEmulatedFrameTime(XrTime Time, int32 FrameNum)
 {
 	// TODO: protect against repeated frame times. would a runtime even do that? I guess it could happen,
 	// but this is for our runtime, and we won't do that for now
@@ -319,12 +772,12 @@ void FOpenXRPoseManager::AddEmulatedFrameTime(XrTime Time, int32 FrameNum)
 	EmulatedBaseTime = FMath::Min(EmulatedBaseTime, Time);
 }
 
-bool FOpenXRPoseManager::DoesActionContainPoseHistory(TStaticArray<ANSICHAR, XR_MAX_ACTION_NAME_SIZE>& ActionName)
+bool FOpenXRActionPoseManager::DoesActionContainPoseHistory(TStaticArray<ANSICHAR, XR_MAX_ACTION_NAME_SIZE>& ActionName)
 {
 	return ActionPoseHistories.Contains(ActionName.GetData());
 }
 
-XrSpaceLocation FOpenXRPoseManager::GetEmulatedPoseForTime(XrSpace LocatingSpace, XrSpace BaseSpace, XrTime Time)
+XrSpaceLocation FOpenXRActionPoseManager::GetEmulatedPoseForTime(XrSpace LocatingSpace, XrSpace BaseSpace, XrTime Time)
 {
 	XrSpaceLocation DummyLocation;
 	DummyLocation.locationFlags = XR_SPACE_LOCATION_ORIENTATION_VALID_BIT | XR_SPACE_LOCATION_POSITION_VALID_BIT | XR_SPACE_LOCATION_ORIENTATION_TRACKED_BIT | XR_SPACE_LOCATION_POSITION_TRACKED_BIT;
@@ -341,10 +794,7 @@ XrSpaceLocation FOpenXRPoseManager::GetEmulatedPoseForTime(XrSpace LocatingSpace
 	int32 PoseIndex = 0;
 	if (XRScribePoseReplayMode == EXRScribePoseReplayMode::TimeMatched)
 	{
-		const XrTime TimeOffsetFromSessionStart = Time - EmulatedBaseTime;
-		const int64 SliceOffset = TimeOffsetFromSessionStart/ 8333333;
-		const int64 SliceIndex = SliceOffset % CapturedSliceCount;
-		PoseIndex = SliceIndex;		
+		PoseIndex = GenerateReplaySliceIndexFromTime(Time);
 	}
 	else
 	{
@@ -394,4 +844,140 @@ XrSpaceLocation FOpenXRPoseManager::GetEmulatedPoseForTime(XrSpace LocatingSpace
 	return DummyLocation;
 }
 
+XrActionStateBoolean FOpenXRActionPoseManager::GetEmulatedActionStateBoolean(const XrActionStateGetInfo* GetInfo, XrTime LastSyncTime)
+{
+	XrActionStateBoolean EmulatedState{};
+	EmulatedState.type = XR_TYPE_ACTION_STATE_BOOLEAN;
+	EmulatedState.lastChangeTime = INT64_MAX;
+
+	const FName ActionName = GetEmulatedActionName(GetInfo->action);
+	if (BooleanProcessedHistories.Contains(ActionName) &&
+		ValidateSubpath(BooleanProcessedHistories[ActionName].ValidSubpaths, GetInfo->subactionPath))
+	{
+		const int64 SliceIndex = GenerateReplaySliceIndexFromTime(LastSyncTime);
+
+		FBooleanActionState HistoryEntry = BooleanProcessedHistories[ActionName].StateValues[SliceIndex];
+		EmulatedState.isActive = HistoryEntry.Active;
+
+		if (EmulatedState.isActive == XR_TRUE)
+		{
+			EmulatedState.currentState = HistoryEntry.State;
+
+			const XrActionStateBoolean CachedState = GetCachedEmulatedBooleanState(ActionName);
+			if (CachedState.isActive == XR_FALSE)
+			{
+				// TODO: this is the 'first' active entry seen, so we need to just cache it.
+				// Can we do this without using this branch?
+				EmulatedState.changedSinceLastSync = XR_FALSE;
+				EmulatedState.lastChangeTime = LastSyncTime;
+				SetCachedEmulatedBooleanState(ActionName, EmulatedState);
+			}
+			else if (CachedState.currentState != EmulatedState.currentState)
+			{
+				EmulatedState.changedSinceLastSync = XR_TRUE;
+				EmulatedState.lastChangeTime = LastSyncTime;
+				SetCachedEmulatedBooleanState(ActionName, EmulatedState);
+			}
+			else
+			{
+				EmulatedState = CachedState;
+			}
+		}
+		// TODO: cache off inactive states?
+	}
+	
+	return EmulatedState;
+}
+
+XrActionStateFloat FOpenXRActionPoseManager::GetEmulatedActionStateFloat(const XrActionStateGetInfo* GetInfo, XrTime LastSyncTime)
+{
+	XrActionStateFloat EmulatedState{};
+	EmulatedState.type = XR_TYPE_ACTION_STATE_FLOAT;
+	EmulatedState.lastChangeTime = INT64_MAX;
+
+	const FName ActionName = GetEmulatedActionName(GetInfo->action);
+	if (FloatProcessedHistories.Contains(ActionName) &&
+		ValidateSubpath(FloatProcessedHistories[ActionName].ValidSubpaths, GetInfo->subactionPath))
+	{
+		const int64 SliceIndex = GenerateReplaySliceIndexFromTime(LastSyncTime);
+
+		FFloatActionState HistoryEntry = FloatProcessedHistories[ActionName].StateValues[SliceIndex];
+		EmulatedState.isActive = HistoryEntry.Active;
+
+		if (EmulatedState.isActive == XR_TRUE)
+		{
+			EmulatedState.currentState = HistoryEntry.State;
+
+			const XrActionStateFloat CachedState = GetCachedEmulatedFloatState(ActionName);
+			if (CachedState.isActive == XR_FALSE)
+			{
+				// TODO: this is the 'first' active entry seen, so we need to just cache it.
+				// Can we do this without using this branch?
+				EmulatedState.changedSinceLastSync = XR_FALSE;
+				EmulatedState.lastChangeTime = LastSyncTime;
+				SetCachedEmulatedFloatState(ActionName, EmulatedState);
+			}
+			else if (CachedState.currentState != EmulatedState.currentState)
+			{
+				EmulatedState.changedSinceLastSync = XR_TRUE;
+				EmulatedState.lastChangeTime = LastSyncTime;
+				SetCachedEmulatedFloatState(ActionName, EmulatedState);
+			}
+			else
+			{
+				EmulatedState = CachedState;
+			}
+		}
+		// TODO: cache off inactive states?
+	}
+
+	return EmulatedState;
+}
+
+XrActionStateVector2f FOpenXRActionPoseManager::GetEmulatedActionStateVector2f(const XrActionStateGetInfo* GetInfo, XrTime LastSyncTime)
+{
+	XrActionStateVector2f EmulatedState{};
+	EmulatedState.type = XR_TYPE_ACTION_STATE_VECTOR2F;
+	EmulatedState.lastChangeTime = INT64_MAX;
+
+	const FName ActionName = GetEmulatedActionName(GetInfo->action);
+	if (Vector2fProcessedHistories.Contains(ActionName) &&
+		ValidateSubpath(Vector2fProcessedHistories[ActionName].ValidSubpaths, GetInfo->subactionPath))
+	{
+		const int64 SliceIndex = GenerateReplaySliceIndexFromTime(LastSyncTime);
+
+		FVector2fActionState HistoryEntry = Vector2fProcessedHistories[ActionName].StateValues[SliceIndex];
+		EmulatedState.isActive = HistoryEntry.Active;
+
+		if (EmulatedState.isActive == XR_TRUE)
+		{
+			EmulatedState.currentState = HistoryEntry.State;
+
+			const XrActionStateVector2f CachedState = GetCachedEmulatedVector2fState(ActionName);
+			if (CachedState.isActive == XR_FALSE)
+			{
+				// TODO: this is the 'first' active entry seen, so we need to just cache it.
+				// Can we do this without using this branch?
+				EmulatedState.changedSinceLastSync = XR_FALSE;
+				EmulatedState.lastChangeTime = LastSyncTime;
+				SetCachedEmulatedVector2fState(ActionName, EmulatedState);
+			}
+			else if ((CachedState.currentState.x != EmulatedState.currentState.x) ||
+				(CachedState.currentState.y != EmulatedState.currentState.y))
+			{
+				EmulatedState.changedSinceLastSync = XR_TRUE;
+				EmulatedState.lastChangeTime = LastSyncTime;
+				SetCachedEmulatedVector2fState(ActionName, EmulatedState);
+			}
+			else
+			{
+				EmulatedState = CachedState;
+			}
+		}
+		// TODO: cache off inactive states?
+	}
+
+	return EmulatedState;
+}
+
 } // namespace UE::XRScribe
diff --git a/Engine/Plugins/Runtime/XR/XRScribe/Source/XRScribe/Private/XRScribeEmulationLayer.cpp b/Engine/Plugins/Runtime/XR/XRScribe/Source/XRScribe/Private/XRScribeEmulationLayer.cpp
--- a/Engine/Plugins/Runtime/XR/XRScribe/Source/XRScribe/Private/XRScribeEmulationLayer.cpp
+++ b/Engine/Plugins/Runtime/XR/XRScribe/Source/XRScribe/Private/XRScribeEmulationLayer.cpp
@@ -358,20 +358,23 @@ void FOpenXREmulationLayer::PostLoadActions()
 		FCStringAnsi::Strncpy(&EmulatedSystemProperties.systemName[ActualSystemNameLen], EmuStringAddend, RemainingSpace);
 	}
 
-	PoseManager.RegisterCapturedWaitFrames(CaptureDecoder.GetWaitFrames());
+	ActionPoseManager.RegisterCapturedPathStrings(CaptureDecoder.GetPathToStringMap());
+	ActionPoseManager.RegisterCapturedWaitFrames(CaptureDecoder.GetWaitFrames());
 
-	PoseManager.RegisterCapturedReferenceSpaces(CaptureDecoder.GetCreatedReferenceSpaces());
+	ActionPoseManager.RegisterCapturedReferenceSpaces(CaptureDecoder.GetCreatedReferenceSpaces());
 
-	PoseManager.RegisterCapturedActions(CaptureDecoder.GetCreatedActions());
-	PoseManager.RegisterCapturedActionSpaces(CaptureDecoder.GetCreatedActionSpaces());
+	ActionPoseManager.RegisterCapturedActions(CaptureDecoder.GetCreatedActions());
+	ActionPoseManager.RegisterCapturedActionSpaces(CaptureDecoder.GetCreatedActionSpaces());
 
-	PoseManager.RegisterCapturedSpaceHistories(CaptureDecoder.GetSpaceLocations());
+	ActionPoseManager.RegisterCapturedSpaceHistories(CaptureDecoder.GetSpaceLocations());
 
-	PoseManager.RegisterCapturedActionStates(CaptureDecoder.GetSyncActions(),
+	ActionPoseManager.RegisterCapturedActionStates(CaptureDecoder.GetSyncActions(),
 		CaptureDecoder.GetBooleanActionStates(), 
 		CaptureDecoder.GetFloatActionStates(), 
 		CaptureDecoder.GetVectorActionStates(), 
 		CaptureDecoder.GetPoseActionStates());
+
+	ActionPoseManager.ProcessCapturedHistories();
 }
 
 bool FOpenXREmulationLayer::InstanceHandleCheck(XrInstance InputInstance)
@@ -1015,7 +1018,7 @@ XrResult FOpenXREmulationLayer::XrLayerDestroySession(XrSession session)
 	// from the spec: "The application is responsible for ensuring that it has no calls using session in progress when the session is destroyed."
 	CurrentSession.Reset();
 
-	PoseManager.OnSessionTeardown();
+	ActionPoseManager.OnSessionTeardown();
 
 	return XR_SUCCESS;
 }
@@ -1099,7 +1102,7 @@ XrResult FOpenXREmulationLayer::XrLayerCreateReferenceSpace(XrSession session, c
 		*space = reinterpret_cast<XrSpace>(CurrentSession->ActiveSpaces.Last().Get());
 	}
 
-	PoseManager.RegisterEmulatedReferenceSpace(*createInfo, *space);
+	ActionPoseManager.RegisterEmulatedReferenceSpace(*createInfo, *space);
 
 	return XrResult::XR_SUCCESS;
 }
@@ -1187,7 +1190,7 @@ XrResult FOpenXREmulationLayer::XrLayerCreateActionSpace(XrSession session, cons
 
 	{
 		FOpenXREmulatedAction* OwningAction = reinterpret_cast<FOpenXREmulatedAction*>(createInfo->action);
-		PoseManager.RegisterEmulatedActionSpace(OwningAction->ActionName, *createInfo, *space);
+		ActionPoseManager.RegisterEmulatedActionSpace(OwningAction->ActionName, *createInfo, *space);
 	}
 
 	return XrResult::XR_SUCCESS;
@@ -1213,7 +1216,7 @@ XrResult FOpenXREmulationLayer::XrLayerLocateSpace(XrSpace space, XrSpace baseSp
 		return XR_ERROR_VALIDATION_FAILURE;
 	}
 
-	*location = PoseManager.GetEmulatedPoseForTime(space, baseSpace, time);
+	*location = ActionPoseManager.GetEmulatedPoseForTime(space, baseSpace, time);
 
 	return XrResult::XR_SUCCESS;
 }
@@ -1825,7 +1828,7 @@ XrResult FOpenXREmulationLayer::XrLayerWaitFrame(XrSession session, const XrFram
 	// TODO: i'm just making up 4 frames of latency
 	frameState->predictedDisplayTime = UETicksToXrTime(FDateTime::Now().GetTicks()) + (4 * frameState->predictedDisplayPeriod);
 
-	PoseManager.AddEmulatedFrameTime(frameState->predictedDisplayTime, CurrentSession->WaitFrameCounter);
+	ActionPoseManager.AddEmulatedFrameTime(frameState->predictedDisplayTime, CurrentSession->WaitFrameCounter);
 
 	return XR_SUCCESS;
 }
@@ -2052,6 +2055,7 @@ XrResult FOpenXREmulationLayer::XrLayerStringToPath(XrInstance instance, const c
 			const int32 AddedIndex = CurrentInstance->PathList.Add(MoveTemp(EmulatedPathString));
 			const XrPath AddedPath = AddedIndex;
 			CurrentInstance->ANSIStringToPathMap.Add(pathString, AddedPath);
+			ActionPoseManager.RegisterEmulatedPath(FName(ANSI_TO_TCHAR(pathString)), AddedPath);
 			*path = AddedPath;
 		}
 	}
@@ -2255,6 +2259,8 @@ XrResult FOpenXREmulationLayer::XrLayerCreateAction(XrActionSet actionSet, const
 		EmulatedActionSet->ActionNames.Add(Action->ActionName.GetData());
 		EmulatedActionSet->Actions.Add(MoveTemp(Action));
 		CurrentInstance->AllActiveActions.Add(*action);
+
+		ActionPoseManager.RegisterEmulatedAction(FName(ANSI_TO_TCHAR(createInfo->actionName)), *action, createInfo->actionType);
 	}
 
 	return XrResult::XR_SUCCESS;
@@ -2405,11 +2411,7 @@ XrResult FOpenXREmulationLayer::XrLayerGetActionStateBoolean(XrSession session,
 	// TODO: check action exists and is attached
 	// TODO: check subaction path
 
-	// TODO: fill XrActionStateBoolean with real data
-	state->currentState = XR_FALSE;
-	state->changedSinceLastSync = XR_FALSE;
-	state->lastChangeTime = UETicksToXrTime(LastSyncTimeTicks);
-	state->isActive = XR_FALSE;
+	*state = ActionPoseManager.GetEmulatedActionStateBoolean(getInfo, UETicksToXrTime(LastSyncTimeTicks));
 
 	return XR_SUCCESS;
 }
@@ -2436,11 +2438,7 @@ XrResult FOpenXREmulationLayer::XrLayerGetActionStateFloat(XrSession session, co
 	// TODO: check action exists and is attached
 	// TODO: check subaction path
 
-	// TODO: fill XrActionStateFloat with real data
-	state->currentState = 0.0;
-	state->changedSinceLastSync = XR_FALSE;
-	state->lastChangeTime = UETicksToXrTime(LastSyncTimeTicks);
-	state->isActive = XR_FALSE;
+	*state = ActionPoseManager.GetEmulatedActionStateFloat(getInfo, UETicksToXrTime(LastSyncTimeTicks));
 
 	return XR_SUCCESS;
 }
@@ -2467,11 +2465,8 @@ XrResult FOpenXREmulationLayer::XrLayerGetActionStateVector2f(XrSession session,
 	// TODO: check action exists and is attached
 	// TODO: check subaction path
 
-	// TODO: fill XrActionStateVector2f with real data
-	state->currentState = XrVector2f({0.0f, 0.0f});
-	state->changedSinceLastSync = XR_FALSE;
-	state->lastChangeTime = UETicksToXrTime(LastSyncTimeTicks);
-	state->isActive = XR_FALSE;
+	*state = ActionPoseManager.GetEmulatedActionStateVector2f(getInfo, UETicksToXrTime(LastSyncTimeTicks));
+
 
 	return XR_SUCCESS;
 }
@@ -2514,7 +2509,7 @@ XrResult FOpenXREmulationLayer::XrLayerGetActionStatePose(XrSession session, con
 
 	// TODO: check subaction path?
 
-	state->isActive = PoseManager.DoesActionContainPoseHistory(Action->ActionName);
+	state->isActive = ActionPoseManager.DoesActionContainPoseHistory(Action->ActionName);
 
 	return XR_SUCCESS;
 }
@@ -2556,7 +2551,8 @@ XrResult FOpenXREmulationLayer::XrLayerSyncActions(XrSession session, const	XrAc
 
 	// TODO: validate the action set subaction path
 
-	LastSyncTimeTicks = FDateTime::Now().GetTicks() - StartTimeTicks;
+	//LastSyncTimeTicks = FDateTime::Now().GetTicks() - StartTimeTicks;
+	LastSyncTimeTicks = FDateTime::Now().GetTicks();
 
 	return XR_SUCCESS;
 }
diff --git a/Engine/Plugins/Runtime/XR/XRScribe/Source/XRScribe/Public/XRScribeAPIDecoder.h b/Engine/Plugins/Runtime/XR/XRScribe/Source/XRScribe/Public/XRScribeAPIDecoder.h
--- a/Engine/Plugins/Runtime/XR/XRScribe/Source/XRScribe/Public/XRScribeAPIDecoder.h
+++ b/Engine/Plugins/Runtime/XR/XRScribe/Source/XRScribe/Public/XRScribeAPIDecoder.h
@@ -4,6 +4,7 @@
 
 #include "OpenXRPlatformRHI.h"
 #include "Containers/Map.h"
+#include "UObject/NameTypes.h"
 #include "Serialization/ArrayReader.h"
 #include "XRScribeFileFormat.h"
 
@@ -45,6 +46,7 @@ public:
 	[[nodiscard]] const TMap<XrAction, TArray<FOpenXRGetActionStateFloatPacket>>& GetFloatActionStates() { return FloatActionStates; }
 	[[nodiscard]] const TMap<XrAction, TArray<FOpenXRGetActionStateVector2fPacket>>& GetVectorActionStates() { return VectorActionStates; }
 	[[nodiscard]] const TMap<XrAction, TArray<FOpenXRGetActionStatePosePacket>>& GetPoseActionStates() { return PoseActionStates; }
+	[[nodiscard]] const TMap<XrPath, FName>& GetPathToStringMap() { return PathToStringMap; }
 
 	TArray<uint8>& GetEncodedData()
 	{
@@ -160,8 +162,8 @@ protected:
 
 	TMap<XrViewConfigurationType, TArray<FOpenXRLocateViewsPacket>> ViewLocations;
 
-	TMap<XrPath, FString> PathToStringMap;
-	TMap<FString, TArray<XrActionSuggestedBinding>> StringToSuggestedBindingsMap;
+	TMap<XrPath, FName> PathToStringMap;
+	TMap<FName, TArray<XrActionSuggestedBinding>> StringToSuggestedBindingsMap;
 
 	TArray<FOpenXRCreateActionPacket> CreatedActions;
 
diff --git a/Engine/Plugins/Runtime/XR/XRScribe/Source/XRScribe/Public/XRScribeEmulatedPoseManager.h b/Engine/Plugins/Runtime/XR/XRScribe/Source/XRScribe/Public/XRScribeEmulatedPoseManager.h
--- a/Engine/Plugins/Runtime/XR/XRScribe/Source/XRScribe/Public/XRScribeEmulatedPoseManager.h
+++ b/Engine/Plugins/Runtime/XR/XRScribe/Source/XRScribe/Public/XRScribeEmulatedPoseManager.h
@@ -14,18 +14,18 @@
 namespace UE::XRScribe
 {
 
-// The PoseManager has two primary jobs:
+// The ActionPoseManager has two primary jobs:
 // * ingest data from capture
-// * generate estimated poses based on emulator frametimes
+// * generate estimated poses/actions based on captured state + timings
 //
-// When using the PoseManager, we'll interact with it similarly to xrLocateSpace: give it time and spaces,
-// and return locations.
+// When using the ActionPoseManager, we'll interact with it similarly to xrLocateSpace: give it time and spaces,
+// and return locations. On the Action side, give it an action and sync time, and we return an action state.
 //
 // Internally, we'll plan to have all managed poses in tracker-space, and convert to the different spaces
 // as needed. This does mean we'll need to know how to convert between different spaces, even if the
 // original capture might not tell us how to
 
-class FOpenXRPoseManager
+class FOpenXRActionPoseManager
 {
 
 public:
@@ -33,6 +33,14 @@ public:
 	// Number of poses the manager will keep available relative to 'current' pose, in order for multi-frame pose queries
 	static const int32 PoseHistorySize = 5;
 
+	/**
+	 * Register information about paths to path strings from the capture
+	 *
+	 * @param PathStringMap Captured map of path handles to FNames
+	 *
+	 */
+	void RegisterCapturedPathStrings(const TMap<XrPath, FName>& PathStringMap);
+
 	/**
 	 * Register information about wait frames from a capture file, in order to build a history of
 	 * display times
@@ -94,6 +102,31 @@ public:
 		const TMap<XrAction, TArray<FOpenXRGetActionStateVector2fPacket>>& VectorActionStates,
 		const TMap<XrAction, TArray<FOpenXRGetActionStatePosePacket>>& PoseActionStates);
 
+	/**
+	 * Process all captured pose and action state histories at once, in order to create a unified timeline
+	 *
+	 */
+	void ProcessCapturedHistories();
+
+	/**
+	 * Register emulated path. Manager needs this to match paths between capture and emulation.
+	 *
+	 * @param PathString String used to create path
+	 * @param Path Emulation-created path handle
+	 *
+	 */
+	void RegisterEmulatedPath(FName PathString, XrPath Path);
+
+	/**
+	 * Register emulated action. Manager needs this to match emulated actions against captured action state histories
+	 *
+	 * @param ActionName String used for action name
+	 * @param Action Emulation-created action handle
+	 * @param ActionType Action-type to facilitate history type lookup
+	 *
+	 */
+	void RegisterEmulatedAction(FName ActionName, XrAction Action, XrActionType ActionType);
+
 	/**
 	 * Register emulated reference space. Pose manager needs this to match the new emulated reference space against
 	 * a captured reference space of similar attributes.
@@ -131,7 +164,8 @@ public:
 	 * tell the client that the action is active, and poses are ready to fetch
 	 *
 	 * @param ActionName Name from action create info
-	 *
+	 * 
+	 * return ActionName has available pose history
 	 */
 	bool DoesActionContainPoseHistory(TStaticArray<ANSICHAR, XR_MAX_ACTION_NAME_SIZE>& ActionName);
 
@@ -146,6 +180,19 @@ public:
 	 */
 	XrSpaceLocation GetEmulatedPoseForTime(XrSpace LocatingSpace, XrSpace BaseSpace, XrTime Time);
 
+	/**
+	 * Fetch an action state generated from the captured history, by using XrActionStateGetInfo and most recent sync time
+	 * The function itself is used to designate between different action state histories (bool, float, vector2f). 
+	 *
+	 * @param GetInfo provides an action handle and subpath
+	 * @param LastSyncTime The sync time is managed by the emulation layer, and used to offset into the captured history.
+	 *
+	 * return Action state derived from captured history if available, or default inactive state if action/subpath aren't available.
+	 */
+	XrActionStateBoolean GetEmulatedActionStateBoolean(const XrActionStateGetInfo* GetInfo, XrTime LastSyncTime);
+	XrActionStateFloat GetEmulatedActionStateFloat(const XrActionStateGetInfo* GetInfo, XrTime LastSyncTime);
+	XrActionStateVector2f GetEmulatedActionStateVector2f(const XrActionStateGetInfo* GetInfo, XrTime LastSyncTime);
+
 	/**
 	 * Reset internal emulation state based on session teardown.
 	 *
@@ -157,15 +204,45 @@ public:
 	
 	// TODO: add LocateViews to pose manager tracking
 
-private:
+protected:
+	
+	template <typename StateType>
+	struct TActionState
+	{
+		StateType State{};
+		XrBool32 Active = 0;
+	};
+
+	template <typename StateType>
+	struct TActionStateHistory
+	{
+		TArray<TActionState<StateType>> StateValues;
+		TSet<FName> ValidSubpaths;
+	};
+
+	using FBooleanActionState = TActionState<XrBool32>;
+	using FFloatActionState = TActionState<float>;
+	using FVector2fActionState = TActionState<XrVector2f>;
+
+	using FBooleanActionStateHistory = TActionStateHistory<XrBool32>;
+	using FFloatActionStateHistory = TActionStateHistory<float>;
+	using FVector2fActionStateHistory = TActionStateHistory<XrVector2f>;
 
 	/**
 	 * Helper function to fetch underlying reference space type from opaque XrSpace handle
 	 */
 	XrReferenceSpaceType GetCapturedReferenceSpaceType(XrSpace CapturedSpace);
 
+	/**
+	 * Helper function to validate captured space is an action space
+	 */
 	bool VerifyCapturedActionSpace(XrSpace CapturedSpace);
 
+	/**
+	 * Helper function to determine full time range of capture
+	 */
+	void CalculateCapturedTimeRange();
+
 	/**
 	 * sort and remove duplicates in history for single space
 	 */
@@ -179,13 +256,48 @@ private:
 	void ProcessCapturedReferenceSpaceHistory(const TArray<FOpenXRLocateSpacePacket>& SpaceHistory);
 	void ProcessCapturedActionSpaceHistory(const TArray<FOpenXRLocateSpacePacket>& SpaceHistory);
 
+	FName GetEmulatedActionName(XrAction Action);
+
+	/**
+	 * Helper function to generate index into captured history from emulated time (offset from emulated session start)
+	 */
+	int64 GenerateReplaySliceIndexFromTime(XrTime EmulatedTime);
+
+	bool ValidateSubpath(const TSet<FName>& ValidSubpaths, XrPath ActionSubPath);
+
+	/**
+	 * Helper function to managed 'cached' action states, because we're only supposed to update on state changes
+	 */
+	const XrActionStateBoolean GetCachedEmulatedBooleanState(FName ActionName);
+	void SetCachedEmulatedBooleanState(FName ActionName, const XrActionStateBoolean& EmulatedState);
+	const XrActionStateFloat GetCachedEmulatedFloatState(FName ActionName);
+	void SetCachedEmulatedFloatState(FName ActionName, const XrActionStateFloat& EmulatedState);
+	const XrActionStateVector2f GetCachedEmulatedVector2fState(FName ActionName);
+	void SetCachedEmulatedVector2fState(FName ActionName, const XrActionStateVector2f& EmulatedState);
+
+	TMap<XrPath, FName> CapturedPaths;
+	TMap<XrPath, FName> EmulatedPaths;
+
+	TMap<XrAction, FName> EmulatedActionNames;
+	TMap<FName, XrActionStateBoolean> CachedEmulatedBooleanStates;
+	TMap<FName, XrActionStateFloat> CachedEmulatedFloatStates;
+	TMap<FName, XrActionStateVector2f> CachedEmulatedVector2fStates;
+
 	// Known spaces and relevant actions from capture
 	TArray<FOpenXRCreateReferenceSpacePacket> CapturedReferenceSpaces;
 	TArray<FOpenXRCreateActionSpacePacket> CapturedActionSpaces;
 	TMap<XrAction, FOpenXRCreateActionPacket> CapturedActions;
 
+	TMap<XrSpace, TArray<FOpenXRLocateSpacePacket>> CapturedRawSpaceHistories;
+
+	//TArray<FOpenXRSyncActionsPacket> CapturedRawSyncActionsPackets;
+	TMap<XrAction, TArray<FOpenXRGetActionStateBooleanPacket>> CapturedRawBooleanActionStates;
+	TMap<XrAction, TArray<FOpenXRGetActionStateFloatPacket>> CapturedRawFloatActionStates;
+	TMap<XrAction, TArray<FOpenXRGetActionStateVector2fPacket>> CapturedRawVectorActionStates;
+	TMap<XrAction, TArray<FOpenXRGetActionStatePosePacket>> CapturedRawPoseActionStates;
+
 	/**
-	 * Pose histories for space types, which could be unified across different space handles
+	 * Processed pose histories for space types, which could be unified across different space handles
 	 * with the same underlying reference space type.
 	 */
 	TMap<XrReferenceSpaceType, TArray<FOpenXRLocateSpacePacket>> ReferencePoseHistories;
@@ -196,6 +308,10 @@ private:
 	TMap<FName, TArray<FOpenXRGetActionStateVector2fPacket>> VectorActionStateHistories;
 	TMap<FName, TArray<FOpenXRGetActionStatePosePacket>> PoseActionStateHistories;
 
+	TMap<FName, FBooleanActionStateHistory>  BooleanProcessedHistories;
+	TMap<FName, FFloatActionStateHistory>	 FloatProcessedHistories;
+	TMap<FName, FVector2fActionStateHistory> Vector2fProcessedHistories;
+
 	// Map of emulated spaces to their underlying reference space type (if applicable)
 	TMap<XrSpace, XrReferenceSpaceType> EmulatedReferenceSpaceTypeMap;
 
@@ -212,10 +328,11 @@ private:
 	int32 LastInsertedFrameIndex = 0;
 
 	int64 CapturedSliceCount = 0;
-	XrTime WaitFrameStart = 0;
-	XrTime WaitFrameEnd = 0;
 	TArray<FOpenXRWaitFramePacket> WaitFrameHistory;
 
+	XrTime CapturedRangeStart = 0;
+	XrTime CapturedRangeEnd = 0;
+
 	XrTime EmulatedBaseTime = INT64_MAX;
 };
 
diff --git a/Engine/Plugins/Runtime/XR/XRScribe/Source/XRScribe/Public/XRScribeEmulationLayer.h b/Engine/Plugins/Runtime/XR/XRScribe/Source/XRScribe/Public/XRScribeEmulationLayer.h
--- a/Engine/Plugins/Runtime/XR/XRScribe/Source/XRScribe/Public/XRScribeEmulationLayer.h
+++ b/Engine/Plugins/Runtime/XR/XRScribe/Source/XRScribe/Public/XRScribeEmulationLayer.h
@@ -162,7 +162,7 @@ protected:
 	//TMpscQueue<FOpenXRUpdateEvent> InternalUpdateQueue;
 	TSpscQueue<FOpenXRUpdateEvent> PendingApplicationUpdateQueue;
 	
-	FOpenXRPoseManager PoseManager;
+	FOpenXRActionPoseManager ActionPoseManager;
 };
 
 } // namespace UE::XRScribe
