XRScribe: filter only changedSinceLastSync actions

#rnx

diff --git a/Engine/Plugins/Runtime/XR/XRScribe/Source/XRScribe/Private/XRScribeEmulatedPoseManager.cpp b/Engine/Plugins/Runtime/XR/XRScribe/Source/XRScribe/Private/XRScribeEmulatedPoseManager.cpp
--- a/Engine/Plugins/Runtime/XR/XRScribe/Source/XRScribe/Private/XRScribeEmulatedPoseManager.cpp
+++ b/Engine/Plugins/Runtime/XR/XRScribe/Source/XRScribe/Private/XRScribeEmulatedPoseManager.cpp
@@ -243,7 +243,7 @@ void FOpenXRActionPoseManager::CalculateCapturedTimeRange()
 		uint32 NumActive = 0;
 		for (const FOpenXRGetActionStateBooleanPacket& BooleanActionState : ActionStateHistory.Value)
 		{
-			if (BooleanActionState.BooleanState.isActive)
+			if (BooleanActionState.BooleanState.isActive && BooleanActionState.BooleanState.changedSinceLastSync)
 			{
 				EarliestActionTime = FMath::Min(EarliestActionTime, BooleanActionState.BooleanState.lastChangeTime);
 				LatestActionTime = FMath::Max(LatestActionTime, BooleanActionState.BooleanState.lastChangeTime);
@@ -257,7 +257,7 @@ void FOpenXRActionPoseManager::CalculateCapturedTimeRange()
 		uint32 NumActive = 0;
 		for (const FOpenXRGetActionStateFloatPacket& FloatActionState : ActionStateHistory.Value)
 		{
-			if (FloatActionState.FloatState.isActive)
+			if (FloatActionState.FloatState.isActive && FloatActionState.FloatState.changedSinceLastSync)
 			{
 				EarliestActionTime = FMath::Min(EarliestActionTime, FloatActionState.FloatState.lastChangeTime);
 				LatestActionTime = FMath::Max(LatestActionTime, FloatActionState.FloatState.lastChangeTime);
@@ -271,7 +271,7 @@ void FOpenXRActionPoseManager::CalculateCapturedTimeRange()
 		uint32 NumActive = 0;
 		for (const FOpenXRGetActionStateVector2fPacket& VectorActionState : ActionStateHistory.Value)
 		{
-			if (VectorActionState.Vector2fState.isActive)
+			if (VectorActionState.Vector2fState.isActive && VectorActionState.Vector2fState.changedSinceLastSync)
 			{
 				EarliestActionTime = FMath::Min(EarliestActionTime, VectorActionState.Vector2fState.lastChangeTime);
 				LatestActionTime = FMath::Max(LatestActionTime, VectorActionState.Vector2fState.lastChangeTime);
@@ -345,7 +345,7 @@ void FOpenXRActionPoseManager::ProcessCapturedHistories()
 		// TODO: Scan for any valid subpaths
 		for (const FOpenXRGetActionStateBooleanPacket& Packet : BooleanActionStateList.Value)
 		{
-			if (Packet.BooleanState.isActive == XR_TRUE)
+			if (Packet.BooleanState.isActive == XR_TRUE && Packet.BooleanState.changedSinceLastSync)
 			{
 				ActiveBooleanStateList.Add(Packet);
 				if (Packet.GetInfoBoolean.subactionPath != XR_NULL_PATH)
@@ -445,7 +445,7 @@ void FOpenXRActionPoseManager::ProcessCapturedHistories()
 
 		for (const FOpenXRGetActionStateFloatPacket& Packet : FloatActionStateList.Value)
 		{
-			if (Packet.FloatState.isActive == XR_TRUE)
+			if (Packet.FloatState.isActive == XR_TRUE && Packet.FloatState.changedSinceLastSync)
 			{
 				ActiveFloatStateList.Add(Packet);
 				if (Packet.GetInfoFloat.subactionPath != XR_NULL_PATH)
@@ -533,7 +533,7 @@ void FOpenXRActionPoseManager::ProcessCapturedHistories()
 
 		for (const FOpenXRGetActionStateVector2fPacket& Packet : VectorActionStateList.Value)
 		{
-			if (Packet.Vector2fState.isActive == XR_TRUE)
+			if (Packet.Vector2fState.isActive == XR_TRUE && Packet.Vector2fState.changedSinceLastSync)
 			{
 				ActiveVector2fStateList.Add(Packet);
 				if (Packet.GetInfoVector2f.subactionPath != XR_NULL_PATH)
