XRScribe: filter out nonsensical lastChangeTime for ActionStates

Some runtimes emit lastChangeTime times that are _minutes_ outside of the action session time.
We have some estimated caps on those early times, in order to create a reasonable replay run.
Controlled by CVar XRScribe.ConstrainActionTime

diff --git a/Engine/Plugins/Runtime/XR/XRScribe/Source/XRScribe/Private/XRScribeEmulatedPoseManager.cpp b/Engine/Plugins/Runtime/XR/XRScribe/Source/XRScribe/Private/XRScribeEmulatedPoseManager.cpp
--- a/Engine/Plugins/Runtime/XR/XRScribe/Source/XRScribe/Private/XRScribeEmulatedPoseManager.cpp
+++ b/Engine/Plugins/Runtime/XR/XRScribe/Source/XRScribe/Private/XRScribeEmulatedPoseManager.cpp
@@ -24,6 +24,12 @@ enum class EXRScribePoseReplayMode : int32
  	reinterpret_cast<int32&>(XRScribePoseReplayMode),
  	TEXT("Toggle the pose replay mode for XRScribe (TimeMatched, Immediate)"),
  	ECVF_ReadOnly);
+
+ static int32 ConstrainActionTime = 5;
+ static FAutoConsoleVariableRef CVarXRScribeConstrainActionTime(TEXT("XRScribe.ConstrainActionTime"),
+	 ConstrainActionTime,
+	 TEXT("Limit the earliest action time to be within N frames of the first predictedDisplayTime, -1 disables"),
+	 ECVF_ReadOnly);
 	
  void FOpenXRActionPoseManager::RegisterCapturedPathStrings(const TMap<XrPath, FName>& PathStringMap)
  {
@@ -283,6 +289,11 @@ void FOpenXRActionPoseManager::CalculateCapturedTimeRange()
 
 	CapturedRangeStart = FMath::Min(EarliestActionTime, WaitFrameStart);
 	CapturedRangeEnd = FMath::Max(LatestActionTime, WaitFrameEnd);
+
+	const XrTime ConstrainedEarliestTime = (ConstrainActionTime >= 0) ?
+		(WaitFrameStart - (WaitFrameHistory[0].FrameState.predictedDisplayPeriod * ConstrainActionTime)) :
+		0;
+	CapturedRangeStart = FMath::Max(CapturedRangeStart, ConstrainedEarliestTime);
 }
 
 void FOpenXRActionPoseManager::ProcessCapturedHistories()
@@ -321,7 +332,8 @@ void FOpenXRActionPoseManager::ProcessCapturedHistories()
 
 	// Converting raw action state captures into a processed state history
 	// * Extract valid subpaths from captured state
-	// * Filter out 'inactive' state values
+	// * Filter out 'invalid' state values
+	//   * invalid = inactive, changedSinceLastSync is false, or lastChangeTime is 'unreasonably' early 
 	//   * We might want to support intermediate inactive periods in the future
 	//	   Currently, we do not support 'holes' in the state history.
 	// * Compress list of active states to reflect changes only (based on change time)
@@ -345,7 +357,9 @@ void FOpenXRActionPoseManager::ProcessCapturedHistories()
 		// TODO: Scan for any valid subpaths
 		for (const FOpenXRGetActionStateBooleanPacket& Packet : BooleanActionStateList.Value)
 		{
-			if (Packet.BooleanState.isActive == XR_TRUE && Packet.BooleanState.changedSinceLastSync)
+			if (Packet.BooleanState.isActive == XR_TRUE && 
+				Packet.BooleanState.changedSinceLastSync &&
+				Packet.BooleanState.lastChangeTime >= CapturedRangeStart)
 			{
 				ActiveBooleanStateList.Add(Packet);
 				if (Packet.GetInfoBoolean.subactionPath != XR_NULL_PATH)
@@ -445,7 +459,9 @@ void FOpenXRActionPoseManager::ProcessCapturedHistories()
 
 		for (const FOpenXRGetActionStateFloatPacket& Packet : FloatActionStateList.Value)
 		{
-			if (Packet.FloatState.isActive == XR_TRUE && Packet.FloatState.changedSinceLastSync)
+			if (Packet.FloatState.isActive == XR_TRUE && 
+				Packet.FloatState.changedSinceLastSync &&
+				Packet.FloatState.lastChangeTime >= CapturedRangeStart)
 			{
 				ActiveFloatStateList.Add(Packet);
 				if (Packet.GetInfoFloat.subactionPath != XR_NULL_PATH)
@@ -533,7 +549,9 @@ void FOpenXRActionPoseManager::ProcessCapturedHistories()
 
 		for (const FOpenXRGetActionStateVector2fPacket& Packet : VectorActionStateList.Value)
 		{
-			if (Packet.Vector2fState.isActive == XR_TRUE && Packet.Vector2fState.changedSinceLastSync)
+			if (Packet.Vector2fState.isActive == XR_TRUE && 
+				Packet.Vector2fState.changedSinceLastSync &&
+				Packet.Vector2fState.lastChangeTime >= CapturedRangeStart)
 			{
 				ActiveVector2fStateList.Add(Packet);
 				if (Packet.GetInfoVector2f.subactionPath != XR_NULL_PATH)
