Fix for XRScribe crashing when decoding a capture where xrBeginFrame returned XR_FRAME_DISCARDED, which is considered a success return and will not impact emulation.

diff --git a/Engine/Plugins/Runtime/XR/XRScribe/Source/XRScribe/Private/XRScribeAPIDecoder.cpp b/Engine/Plugins/Runtime/XR/XRScribe/Source/XRScribe/Private/XRScribeAPIDecoder.cpp
--- a/Engine/Plugins/Runtime/XR/XRScribe/Source/XRScribe/Private/XRScribeAPIDecoder.cpp
+++ b/Engine/Plugins/Runtime/XR/XRScribe/Source/XRScribe/Private/XRScribeAPIDecoder.cpp
@@ -521,7 +521,7 @@ void FOpenXRCaptureDecoder::DecodeWaitFrame(const FOpenXRAPIPacketBase& BasePack
 void FOpenXRCaptureDecoder::DecodeBeginFrame(const FOpenXRAPIPacketBase& BasePacket)
 {
 	check(BasePacket.ApiId == EOpenXRAPIPacketId::BeginFrame);
-	check(BasePacket.Result == XR_SUCCESS);
+	check(BasePacket.Result == XR_SUCCESS || BasePacket.Result == XR_FRAME_DISCARDED);
 
 	FOpenXRBeginFramePacket Data(XrResult::XR_ERROR_RUNTIME_FAILURE);
 
