Add a plugin settings page for the XRScribe plugin with a dropdown to select Capture or Emulate mode (requires restart).

diff --git a/Engine/Plugins/Runtime/XR/XRScribe/Source/XRScribe/Private/XRScribeAPISurface.cpp b/Engine/Plugins/Runtime/XR/XRScribe/Source/XRScribe/Private/XRScribeAPISurface.cpp
--- a/Engine/Plugins/Runtime/XR/XRScribe/Source/XRScribe/Private/XRScribeAPISurface.cpp
+++ b/Engine/Plugins/Runtime/XR/XRScribe/Source/XRScribe/Private/XRScribeAPISurface.cpp
@@ -25,8 +25,10 @@ namespace UE::XRScribe
 // know the value super early, in order to re-direct the initial OpenXR API calls
 // Example addition to DefaultEngine.ini
 //
-//	[SystemSettings]
+//	[/Script/XRScribe.XRScribeDeveloperSettings]
 //	XRScribe.RunMode = 1
+//
+// Can be set manually, or via the plugin's project settings page. Defaults to 1 (emulate).
 
 enum class EXRScribeRunMode : int32
 {
@@ -690,7 +692,7 @@ EXRScribeRunMode DetermineRunMode(EXRScribeRunMode FallbackRunMode, FString& Emu
 	// TODO: load file from plugin settings
 
 	EXRScribeRunMode RunModeFromConfig = EXRScribeRunMode::Emulate;
-	const bool bModeReadFromConfig = GConfig->GetInt(TEXT("SystemSettings"), TEXT("XRScribe.RunMode"), (int32&)RunModeFromConfig, GEngineIni);
+	const bool bModeReadFromConfig = GConfig->GetInt(TEXT("/Script/XRScribe.XRScribeDeveloperSettings"), TEXT("XRScribe.RunMode"), (int32&)RunModeFromConfig, GEngineIni);
 	if (!bModeReadFromConfig)
 	{
 		RunModeFromConfig = FallbackRunMode;
diff --git a/Engine/Plugins/Runtime/XR/XRScribe/Source/XRScribe/Private/XRScribeDeveloperSettings.cpp b/Engine/Plugins/Runtime/XR/XRScribe/Source/XRScribe/Private/XRScribeDeveloperSettings.cpp
--- a/Engine/Plugins/Runtime/XR/XRScribe/Source/XRScribe/Private/XRScribeDeveloperSettings.cpp
+++ b/Engine/Plugins/Runtime/XR/XRScribe/Source/XRScribe/Private/XRScribeDeveloperSettings.cpp
@@ -2,9 +2,14 @@
 
 #include "XRScribeDeveloperSettings.h"
 
-//#include UE_INLINE_GENERATED_CPP_BY_NAME(XRScribeDeveloperSettings)
-//
-//UXRScribeDeveloperSettings::UXRScribeDeveloperSettings(const FObjectInitializer& Initializer)
-//	: Super(Initializer)
-//{
-//}
+#include UE_INLINE_GENERATED_CPP_BY_NAME(XRScribeDeveloperSettings)
+
+UXRScribeDeveloperSettings::UXRScribeDeveloperSettings(const FObjectInitializer& Initializer)
+	: Super(Initializer)
+{
+}
+
+FName UXRScribeDeveloperSettings::GetCategoryName() const
+{
+	return FName(TEXT("Plugins"));
+}
\ No newline at end of file
diff --git a/Engine/Plugins/Runtime/XR/XRScribe/Source/XRScribe/Public/XRScribeDeveloperSettings.h b/Engine/Plugins/Runtime/XR/XRScribe/Source/XRScribe/Public/XRScribeDeveloperSettings.h
--- a/Engine/Plugins/Runtime/XR/XRScribe/Source/XRScribe/Public/XRScribeDeveloperSettings.h
+++ b/Engine/Plugins/Runtime/XR/XRScribe/Source/XRScribe/Public/XRScribeDeveloperSettings.h
@@ -2,29 +2,48 @@
 
 #pragma once
 
-// TODO: Use UDeveloperSettings + Project Settings UI in editor
-// Ideally, we'd use CDO UXRScribeDeveloperSettings to fetch our dev settings
-// however, we have some issues now with plugin loading + uobject setup
-// Even if we 'fix' those issues, reading the config into the class isn't stood up this early
-
-//#include "Engine/DeveloperSettings.h"
-//#include "XRScribeDeveloperSettings.generated.h"
-//
-///** Developer settings for XRScribe */
-//UCLASS(config = Engine, defaultconfig, meta = (DisplayName = "XRScribe"))
-//class XRSCRIBE_API UXRScribeDeveloperSettings : public UDeveloperSettings
-//{
-//	GENERATED_BODY()
-//public:
-//
-//	UXRScribeDeveloperSettings(const FObjectInitializer& Initializer);
-//
-//	/** Controls whether XRScribe runs in capture or emulation mode. Currently needed at engine startup, but will be runtime switchable. */
-//	UPROPERTY(config, EditAnywhere, Category = "XRScribe")
-//	EXRScribeRunMode RunMode = EXRScribeRunMode::Capture;
-//
-//	// TODO:
-//	// File path for capture file
-//	// Customizing capture dump point (session end, instance teardown, app end)
-//	// other run modes (e.g. replay)
-//};
\ No newline at end of file
+// Because we initialize XRScribe before developer settings are read in from the config, this setting is just used
+// to set XRScribe.RunMode in DefaultEngine.ini, where it'll be read in directly from the config file in DetermineRunMode().
+// This differs from the normal method of pulling settings from the backing CVars or using the UXRScribeDeveloperSettings 
+// class default object.
+
+#include "Engine/DeveloperSettings.h"
+#include "XRScribeDeveloperSettings.generated.h"
+
+/**
+ * Enumerates available options for XRScribe run mode.
+ */
+UENUM()
+namespace EXRScribeRunMode
+{
+	enum Type : int
+	{
+		Capture = 0 UMETA(ToolTip = "Capture OpenXR API calls and store to Saved/Capture.xrs."),
+		Emulate = 1 UMETA(ToolTip = "Emulate OpenXR runtime and playback Saved/Capture.xrs."),
+	};
+}
+
+/** Developer settings for XRScribe */
+UCLASS(config = Engine, defaultconfig, meta = (DisplayName = "XRScribe"))
+class XRSCRIBE_API UXRScribeDeveloperSettings : public UDeveloperSettings
+{
+	GENERATED_BODY()
+public:
+
+	UXRScribeDeveloperSettings(const FObjectInitializer& Initializer);
+
+	/** Controls whether XRScribe runs in capture or emulation mode. Currently needed at engine startup, but will be runtime switchable. */
+	UPROPERTY(config, EditAnywhere, Category = "XRScribe", meta = (
+		ConsoleVariable = "XRScribe.RunMode", DisplayName = "Run Mode",
+		ToolTip = "0 - Capture, 1 - Emulate",
+		ConfigRestartRequired = true))
+	TEnumAsByte<EXRScribeRunMode::Type> RunMode = EXRScribeRunMode::Type::Emulate; // This should match the default FallbackRunMode set in XRScribeAPISurface.cpp
+
+	// TODO:
+	// File path for capture file
+	// Customizing capture dump point (session end, instance teardown, app end)
+	// other run modes (e.g. replay)
+
+public:
+	virtual FName GetCategoryName() const;
+};
\ No newline at end of file
