When loading a malformed XRScribe capture, log the issue and display a message box confirming that the emulation layer will be disabled rather than crashing the editor.

diff --git a/Engine/Plugins/Runtime/XR/XRScribe/Source/XRScribe/Private/XRScribeAPIDecoder.cpp b/Engine/Plugins/Runtime/XR/XRScribe/Source/XRScribe/Private/XRScribeAPIDecoder.cpp
--- a/Engine/Plugins/Runtime/XR/XRScribe/Source/XRScribe/Private/XRScribeAPIDecoder.cpp
+++ b/Engine/Plugins/Runtime/XR/XRScribe/Source/XRScribe/Private/XRScribeAPIDecoder.cpp
@@ -93,35 +93,58 @@ FOpenXRCaptureDecoder::FOpenXRCaptureDecoder()
 
 FOpenXRCaptureDecoder::~FOpenXRCaptureDecoder() {}
 
-void FOpenXRCaptureDecoder::DecodeDataFromMemory()
+bool FOpenXRCaptureDecoder::DecodeDataFromMemory()
 {
 	const int64 NumBytes = EncodedData.Num();
 	int64 CurByteIndex = 0;
+	int64 CurPacketIndex = 0;
 
 	while (CurByteIndex < NumBytes)
 	{
 		const FOpenXRAPIPacketBase* NextPacket = reinterpret_cast<const FOpenXRAPIPacketBase*>(&EncodedData[CurByteIndex]);
 
-		check(NextPacket->Padding0 == FOpenXRAPIPacketBase::MagicPacketByte);
-		check(NextPacket->ApiId >= EOpenXRAPIPacketId::EnumerateApiLayerProperties && NextPacket->ApiId < EOpenXRAPIPacketId::NumValidAPIPacketIds);
+		if (NextPacket->Padding0 != FOpenXRAPIPacketBase::MagicPacketByte)
+		{
+			UE_LOG(LogXRScribeEmulate, Error, TEXT("Encountered invalid magic byte in packet %d while decoding capture"), CurPacketIndex);
+			return false;
+		}
+
+		if (NextPacket->ApiId < EOpenXRAPIPacketId::EnumerateApiLayerProperties || NextPacket->ApiId > EOpenXRAPIPacketId::NumValidAPIPacketIds)
+		{
+			UE_LOG(LogXRScribeEmulate, Error, TEXT("Encountered invalid API ID in packet %d while decoding capture"), CurPacketIndex);
+			return false;
+		}
+
 		check(DecodeFnTable[(uint32)NextPacket->ApiId] != nullptr);
 
 		// We explicitly need `this` to make it clear which instance is being passed to function pointer
-		(this->*DecodeFnTable[(uint32)NextPacket->ApiId])(*NextPacket); 
+		const bool bValid = (this->*DecodeFnTable[(uint32)NextPacket->ApiId])(*NextPacket);
+		if (!bValid)
+		{
+			UE_LOG(LogXRScribeEmulate, Error, TEXT("Encountered invalid XrResult in packet %d while decoding capture"), CurPacketIndex);
+			return false;
+		}
 
 		CurByteIndex = EncodedData.Tell();
+		CurPacketIndex++;
 	}
 
 	check(CurByteIndex == NumBytes);
+
+	return true;
 }
 
 ///////////////
 /// packet decoders
 
-void FOpenXRCaptureDecoder::DecodeEnumerateApiLayerProperties(const FOpenXRAPIPacketBase& BasePacket)
+bool FOpenXRCaptureDecoder::DecodeEnumerateApiLayerProperties(const FOpenXRAPIPacketBase& BasePacket)
 {
 	check(BasePacket.ApiId == EOpenXRAPIPacketId::EnumerateApiLayerProperties);
-	check(BasePacket.Result == XR_SUCCESS);
+
+	if (BasePacket.Result != XR_SUCCESS)
+	{
+		return false;
+	}
 
 	FOpenXREnumerateApiLayerPropertiesPacket Data(XrResult::XR_ERROR_RUNTIME_FAILURE);
 
@@ -131,29 +154,44 @@ void FOpenXRCaptureDecoder::DecodeEnumerateApiLayerProperties(const FOpenXRAPIPa
 	{
 		ApiLayerProperties = MoveTemp(Data.LayerProperties);
 	}
+
+	return true;
 }
 
-void FOpenXRCaptureDecoder::DecodeEnumerateInstanceExtensionProperties(const FOpenXRAPIPacketBase& BasePacket)
+bool FOpenXRCaptureDecoder::DecodeEnumerateInstanceExtensionProperties(const FOpenXRAPIPacketBase& BasePacket)
 {
 	check(BasePacket.ApiId == EOpenXRAPIPacketId::EnumerateInstanceExtensionProperties);
-	check(BasePacket.Result == XR_SUCCESS);
+
+	if (BasePacket.Result != XR_SUCCESS)
+	{
+		return false;
+	}
 
 	FOpenXREnumerateInstanceExtensionPropertiesPacket Data(XrResult::XR_ERROR_RUNTIME_FAILURE);
 
 	EncodedData << Data;
-	
-	check(Data.LayerName[0] == 0); // TODO: save off extension properties queried from layers
+
+	if (Data.LayerName[0] != 0)
+	{
+		return false;
+	}
 
 	if (!Data.ExtensionProperties.IsEmpty())
 	{
 		InstanceExtensionProperties = MoveTemp(Data.ExtensionProperties);
 	}
+
+	return true;
 }
 
-void FOpenXRCaptureDecoder::DecodeCreateInstance(const FOpenXRAPIPacketBase& BasePacket)
+bool FOpenXRCaptureDecoder::DecodeCreateInstance(const FOpenXRAPIPacketBase& BasePacket)
 {
 	check(BasePacket.ApiId == EOpenXRAPIPacketId::CreateInstance);
-	check(BasePacket.Result == XR_SUCCESS);
+
+	if (BasePacket.Result != XR_SUCCESS)
+	{
+		return false;
+	}
 
 	FOpenXRCreateInstancePacket Data(XrResult::XR_ERROR_RUNTIME_FAILURE);
 
@@ -162,60 +200,86 @@ void FOpenXRCaptureDecoder::DecodeCreateInstance(const FOpenXRAPIPacketBase& Bas
 	ValidInstanceCreateFlags = Data.CreateFlags;
 	RequestedLayerNames = MoveTemp(Data.EnabledLayerNames);
 	RequestedExtensionNames = MoveTemp(Data.EnabledExtensionNames);
+
+	return true;
 }
 
-void FOpenXRCaptureDecoder::DecodeDestroyInstance(const FOpenXRAPIPacketBase& BasePacket)
+bool FOpenXRCaptureDecoder::DecodeDestroyInstance(const FOpenXRAPIPacketBase& BasePacket)
 {
 	check(BasePacket.ApiId == EOpenXRAPIPacketId::DestroyInstance);
-	check(BasePacket.Result == XR_SUCCESS);
+
+	if (BasePacket.Result != XR_SUCCESS)
+	{
+		return false;
+	}
 
 	FOpenXRDestroyInstancePacket Data(XrResult::XR_ERROR_RUNTIME_FAILURE);
 
 	EncodedData << Data;
+	return true;
 
 	// nothing for us to do with this action currently
 }
 
-void FOpenXRCaptureDecoder::DecodeGetInstanceProperties(const FOpenXRAPIPacketBase& BasePacket)
+bool FOpenXRCaptureDecoder::DecodeGetInstanceProperties(const FOpenXRAPIPacketBase& BasePacket)
 {
 	check(BasePacket.ApiId == EOpenXRAPIPacketId::GetInstanceProperties);
-	check(BasePacket.Result == XR_SUCCESS);
+
+	if (BasePacket.Result != XR_SUCCESS)
+	{
+		return false;
+	}
 
 	FOpenXRGetInstancePropertiesPacket Data(XrResult::XR_ERROR_RUNTIME_FAILURE);
 
 	EncodedData << Data;
 
 	InstanceProperties = MoveTemp(Data.InstanceProperties);
+	return true;
 }
 
-void FOpenXRCaptureDecoder::DecodeGetSystem(const FOpenXRAPIPacketBase& BasePacket)
+bool FOpenXRCaptureDecoder::DecodeGetSystem(const FOpenXRAPIPacketBase& BasePacket)
 {
 	check(BasePacket.ApiId == EOpenXRAPIPacketId::GetSystem);
-	check(BasePacket.Result == XR_SUCCESS);
+
+	if (BasePacket.Result != XR_SUCCESS)
+	{
+		return false;
+	}
 
 	FOpenXRGetSystemPacket Data(XrResult::XR_ERROR_RUNTIME_FAILURE);
 
 	EncodedData << Data;
 
 	SystemGetInfo = MoveTemp(Data.SystemGetInfo);
+	return true;
 }
 
-void FOpenXRCaptureDecoder::DecodeGetSystemProperties(const FOpenXRAPIPacketBase& BasePacket)
+bool FOpenXRCaptureDecoder::DecodeGetSystemProperties(const FOpenXRAPIPacketBase& BasePacket)
 {
 	check(BasePacket.ApiId == EOpenXRAPIPacketId::GetSystemProperties);
-	check(BasePacket.Result == XR_SUCCESS);
+
+	if (BasePacket.Result != XR_SUCCESS)
+	{
+		return false;
+	}
 
 	FOpenXRGetSystemPropertiesPacket Data(XrResult::XR_ERROR_RUNTIME_FAILURE);
 
 	EncodedData << Data;
 
 	SystemProperties = MoveTemp(Data.SystemProperties);
+	return true;
 }
 
-void FOpenXRCaptureDecoder::DecodeEnumerateEnvironmentBlendModes(const FOpenXRAPIPacketBase& BasePacket)
+bool FOpenXRCaptureDecoder::DecodeEnumerateEnvironmentBlendModes(const FOpenXRAPIPacketBase& BasePacket)
 {
 	check(BasePacket.ApiId == EOpenXRAPIPacketId::EnumerateEnvironmentBlendModes);
-	check(BasePacket.Result == XR_SUCCESS);
+
+	if (BasePacket.Result != XR_SUCCESS)
+	{
+		return false;
+	}
 
 	FOpenXREnumerateEnvironmentBlendModesPacket Data(XrResult::XR_ERROR_RUNTIME_FAILURE);
 
@@ -230,36 +294,52 @@ void FOpenXRCaptureDecoder::DecodeEnumerateEnvironmentBlendModes(const FOpenXRAP
 	{
 		// TODO: Log error, unsupported view config type
 	}
+
+	return true;
 }
 
-void FOpenXRCaptureDecoder::DecodeCreateSession(const FOpenXRAPIPacketBase& BasePacket)
+bool FOpenXRCaptureDecoder::DecodeCreateSession(const FOpenXRAPIPacketBase& BasePacket)
 {
 	check(BasePacket.ApiId == EOpenXRAPIPacketId::CreateSession);
-	check(BasePacket.Result == XR_SUCCESS);
+
+	if (BasePacket.Result != XR_SUCCESS)
+	{
+		return false;
+	}
 
 	FOpenXRCreateSessionPacket Data(XrResult::XR_ERROR_RUNTIME_FAILURE);
 
 	EncodedData << Data;
 
 	SessionCreateInfo = Data.SessionCreateInfo;
+	return true;
 }
 
-void FOpenXRCaptureDecoder::DecodeDestroySession(const FOpenXRAPIPacketBase& BasePacket)
+bool FOpenXRCaptureDecoder::DecodeDestroySession(const FOpenXRAPIPacketBase& BasePacket)
 {
 	check(BasePacket.ApiId == EOpenXRAPIPacketId::DestroySession);
-	check(BasePacket.Result == XR_SUCCESS);
+
+	if (BasePacket.Result != XR_SUCCESS)
+	{
+		return false;
+	}
 
 	FOpenXRDestroySessionPacket Data(XrResult::XR_ERROR_RUNTIME_FAILURE);
 
 	EncodedData << Data;
+	return true;
 
 	// Nothing to do
 }
 
-void FOpenXRCaptureDecoder::DecodeEnumerateReferenceSpaces(const FOpenXRAPIPacketBase& BasePacket)
+bool FOpenXRCaptureDecoder::DecodeEnumerateReferenceSpaces(const FOpenXRAPIPacketBase& BasePacket)
 {
 	check(BasePacket.ApiId == EOpenXRAPIPacketId::EnumerateReferenceSpaces);
-	check(BasePacket.Result == XR_SUCCESS);
+
+	if (BasePacket.Result != XR_SUCCESS)
+	{
+		return false;
+	}
 
 	FOpenXREnumerateReferenceSpacesPacket Data(XrResult::XR_ERROR_RUNTIME_FAILURE);
 
@@ -269,12 +349,18 @@ void FOpenXRCaptureDecoder::DecodeEnumerateReferenceSpaces(const FOpenXRAPIPacke
 	{
 		ReferenceSpaceTypes = MoveTemp(Data.Spaces);
 	}
+
+	return true;
 }
 
-void FOpenXRCaptureDecoder::DecodeCreateReferenceSpace(const FOpenXRAPIPacketBase& BasePacket)
+bool FOpenXRCaptureDecoder::DecodeCreateReferenceSpace(const FOpenXRAPIPacketBase& BasePacket)
 {
 	check(BasePacket.ApiId == EOpenXRAPIPacketId::CreateReferenceSpace);
-	check(BasePacket.Result == XR_SUCCESS);
+
+	if (BasePacket.Result != XR_SUCCESS)
+	{
+		return false;
+	}
 
 	FOpenXRCreateReferenceSpacePacket Data(XrResult::XR_ERROR_RUNTIME_FAILURE);
 
@@ -282,12 +368,15 @@ void FOpenXRCaptureDecoder::DecodeCreateReferenceSpace(const FOpenXRAPIPacketBas
 
 	ReferenceSpaceMap.Add(Data.Space, Data.ReferenceSpaceCreateInfo.referenceSpaceType);
 	CreatedReferenceSpaces.Add(Data);
+	return true;
 }
 
-void FOpenXRCaptureDecoder::DecodeGetReferenceSpaceBoundsRect(const FOpenXRAPIPacketBase& BasePacket)
+bool FOpenXRCaptureDecoder::DecodeGetReferenceSpaceBoundsRect(const FOpenXRAPIPacketBase& BasePacket)
 {
 	check(BasePacket.ApiId == EOpenXRAPIPacketId::GetReferenceSpaceBoundsRect);
 
+	// Failed results are allowed here
+
 	FOpenXRGetReferenceSpaceBoundsRectPacket Data(XrResult::XR_ERROR_RUNTIME_FAILURE);
 
 	EncodedData << Data;
@@ -297,12 +386,18 @@ void FOpenXRCaptureDecoder::DecodeGetReferenceSpaceBoundsRect(const FOpenXRAPIPa
 		ReferenceSpaceBounds.Add(Data.ReferenceSpaceType, Data.Bounds);
 		// TODO: Check for existing bounds associated with reference space?
 	}
+
+	return true;
 }
 
-void FOpenXRCaptureDecoder::DecodeCreateActionSpace(const FOpenXRAPIPacketBase& BasePacket)
+bool FOpenXRCaptureDecoder::DecodeCreateActionSpace(const FOpenXRAPIPacketBase& BasePacket)
 {
 	check(BasePacket.ApiId == EOpenXRAPIPacketId::CreateActionSpace);
-	check(BasePacket.Result == XR_SUCCESS);
+
+	if (BasePacket.Result != XR_SUCCESS)
+	{
+		return false;
+	}
 
 	FOpenXRCreateActionSpacePacket Data(XrResult::XR_ERROR_RUNTIME_FAILURE);
 
@@ -310,12 +405,17 @@ void FOpenXRCaptureDecoder::DecodeCreateActionSpace(const FOpenXRAPIPacketBase&
 
 	ActionSpaceMap.Add(Data.Space, Data.ActionSpaceCreateInfo.action);
 	CreatedActionSpaces.Add(Data);
+	return true;
 }
 
-void FOpenXRCaptureDecoder::DecodeLocateSpace(const FOpenXRAPIPacketBase& BasePacket)
+bool FOpenXRCaptureDecoder::DecodeLocateSpace(const FOpenXRAPIPacketBase& BasePacket)
 {
 	check(BasePacket.ApiId == EOpenXRAPIPacketId::LocateSpace);
-	check(BasePacket.Result == XR_SUCCESS);
+
+	if (BasePacket.Result != XR_SUCCESS)
+	{
+		return false;
+	}
 
 	FOpenXRLocateSpacePacket Data(XrResult::XR_ERROR_RUNTIME_FAILURE);
 
@@ -332,24 +432,34 @@ void FOpenXRCaptureDecoder::DecodeLocateSpace(const FOpenXRAPIPacketBase& BasePa
 	//Record.Location = Data.Location;
 
 	SpaceLocations[Data.Space].Add(Data);
+	return true;
 }
 
-void FOpenXRCaptureDecoder::DecodeDestroySpace(const FOpenXRAPIPacketBase& BasePacket)
+bool FOpenXRCaptureDecoder::DecodeDestroySpace(const FOpenXRAPIPacketBase& BasePacket)
 {
 	check(BasePacket.ApiId == EOpenXRAPIPacketId::DestroySpace);
-	check(BasePacket.Result == XR_SUCCESS);
+
+	if (BasePacket.Result != XR_SUCCESS)
+	{
+		return false;
+	}
 
 	FOpenXRDestroySpacePacket Data(XrResult::XR_ERROR_RUNTIME_FAILURE);
 
 	EncodedData << Data;
+	return true;
 
 	// nothing to do yet
 }
 
-void FOpenXRCaptureDecoder::DecodeEnumerateViewConfigurations(const FOpenXRAPIPacketBase& BasePacket)
+bool FOpenXRCaptureDecoder::DecodeEnumerateViewConfigurations(const FOpenXRAPIPacketBase& BasePacket)
 {
 	check(BasePacket.ApiId == EOpenXRAPIPacketId::EnumerateViewConfigurations);
-	check(BasePacket.Result == XR_SUCCESS);
+
+	if (BasePacket.Result != XR_SUCCESS)
+	{
+		return false;
+	}
 
 	FOpenXREnumerateViewConfigurationsPacket Data(XrResult::XR_ERROR_RUNTIME_FAILURE);
 
@@ -359,24 +469,35 @@ void FOpenXRCaptureDecoder::DecodeEnumerateViewConfigurations(const FOpenXRAPIPa
 	{
 		ViewConfigurationTypes = MoveTemp(Data.ViewConfigurationTypes);
 	}
+
+	return true;
 }
 
-void FOpenXRCaptureDecoder::DecodeGetViewConfigurationProperties(const FOpenXRAPIPacketBase& BasePacket)
+bool FOpenXRCaptureDecoder::DecodeGetViewConfigurationProperties(const FOpenXRAPIPacketBase& BasePacket)
 {
 	check(BasePacket.ApiId == EOpenXRAPIPacketId::GetViewConfigurationProperties);
-	check(BasePacket.Result == XR_SUCCESS);
+
+	if (BasePacket.Result != XR_SUCCESS)
+	{
+		return false;
+	}
 
 	FOpenXRGetViewConfigurationPropertiesPacket Data(XrResult::XR_ERROR_RUNTIME_FAILURE);
 
 	EncodedData << Data;
 
 	ViewConfigurationProperties.Add(Data.ViewConfigurationType, Data.ConfigurationProperties);
+	return true;
 }
 
-void FOpenXRCaptureDecoder::DecodeEnumerateViewConfigurationViews(const FOpenXRAPIPacketBase& BasePacket)
+bool FOpenXRCaptureDecoder::DecodeEnumerateViewConfigurationViews(const FOpenXRAPIPacketBase& BasePacket)
 {
 	check(BasePacket.ApiId == EOpenXRAPIPacketId::EnumerateViewConfigurationViews);
-	check(BasePacket.Result == XR_SUCCESS);
+
+	if (BasePacket.Result != XR_SUCCESS)
+	{
+		return false;
+	}
 
 	FOpenXREnumerateViewConfigurationViewsPacket Data(XrResult::XR_ERROR_RUNTIME_FAILURE);
 
@@ -386,12 +507,18 @@ void FOpenXRCaptureDecoder::DecodeEnumerateViewConfigurationViews(const FOpenXRA
 	{
 		ViewConfigurationViews.Add(Data.ViewConfigurationType, Data.Views);
 	}
+
+	return true;
 }
 
-void FOpenXRCaptureDecoder::DecodeEnumerateSwapchainFormats(const FOpenXRAPIPacketBase& BasePacket)
+bool FOpenXRCaptureDecoder::DecodeEnumerateSwapchainFormats(const FOpenXRAPIPacketBase& BasePacket)
 {
 	check(BasePacket.ApiId == EOpenXRAPIPacketId::EnumerateSwapchainFormats);
-	check(BasePacket.Result == XR_SUCCESS);
+
+	if (BasePacket.Result != XR_SUCCESS)
+	{
+		return false;
+	}
 
 	FOpenXREnumerateSwapchainFormatsPacket Data(XrResult::XR_ERROR_RUNTIME_FAILURE);
 
@@ -401,150 +528,221 @@ void FOpenXRCaptureDecoder::DecodeEnumerateSwapchainFormats(const FOpenXRAPIPack
 	{
 		SwapchainFormats = MoveTemp(Data.Formats);
 	}
+
+	return true;
 }
 
-void FOpenXRCaptureDecoder::DecodeCreateSwapchain(const FOpenXRAPIPacketBase& BasePacket)
+bool FOpenXRCaptureDecoder::DecodeCreateSwapchain(const FOpenXRAPIPacketBase& BasePacket)
 {
 	check(BasePacket.ApiId == EOpenXRAPIPacketId::CreateSwapchain);
-	check(BasePacket.Result == XR_SUCCESS);
+
+	if (BasePacket.Result != XR_SUCCESS)
+	{
+		return false;
+	}
 
 	FOpenXRCreateSwapchainPacket Data(XrResult::XR_ERROR_RUNTIME_FAILURE);
 
 	EncodedData << Data;
+	return true;
 
 	// TODO: Actual swapchain creation information not needed...yet
 }
 
-void FOpenXRCaptureDecoder::DecodeDestroySwapchain(const FOpenXRAPIPacketBase& BasePacket)
+bool FOpenXRCaptureDecoder::DecodeDestroySwapchain(const FOpenXRAPIPacketBase& BasePacket)
 {
 	check(BasePacket.ApiId == EOpenXRAPIPacketId::DestroySwapchain);
-	check(BasePacket.Result == XR_SUCCESS);
+
+	if (BasePacket.Result != XR_SUCCESS)
+	{
+		return false;
+	}
 
 	FOpenXRDestroySwapchainPacket Data(XrResult::XR_ERROR_RUNTIME_FAILURE);
 
 	EncodedData << Data;
+	return true;
 
 	// nothing to do yet
 }
 
-void FOpenXRCaptureDecoder::DecodeEnumerateSwapchainImages(const FOpenXRAPIPacketBase& BasePacket)
+bool FOpenXRCaptureDecoder::DecodeEnumerateSwapchainImages(const FOpenXRAPIPacketBase& BasePacket)
 {
 	check(BasePacket.ApiId == EOpenXRAPIPacketId::EnumerateSwapchainImages);
-	check(BasePacket.Result == XR_SUCCESS);
+
+	if (BasePacket.Result != XR_SUCCESS)
+	{
+		return false;
+	}
 
 	FOpenXREnumerateSwapchainImagesPacket Data(XrResult::XR_ERROR_RUNTIME_FAILURE);
 
 	EncodedData << Data;
+	return true;
 
 	// nothing to do yet
 }
 
-void FOpenXRCaptureDecoder::DecodeAcquireSwapchainImage(const FOpenXRAPIPacketBase& BasePacket)
+bool FOpenXRCaptureDecoder::DecodeAcquireSwapchainImage(const FOpenXRAPIPacketBase& BasePacket)
 {
 	check(BasePacket.ApiId == EOpenXRAPIPacketId::AcquireSwapchainImage);
-	check(BasePacket.Result == XR_SUCCESS);
+
+	if (BasePacket.Result != XR_SUCCESS)
+	{
+		return false;
+	}
 
 	FOpenXRAcquireSwapchainImagePacket Data(XrResult::XR_ERROR_RUNTIME_FAILURE);
 
 	EncodedData << Data;
+	return true;
 
 	// nothing to do yet
 }
 
-void FOpenXRCaptureDecoder::DecodeWaitSwapchainImage(const FOpenXRAPIPacketBase& BasePacket)
+bool FOpenXRCaptureDecoder::DecodeWaitSwapchainImage(const FOpenXRAPIPacketBase& BasePacket)
 {
 	check(BasePacket.ApiId == EOpenXRAPIPacketId::WaitSwapchainImage);
-	check(BasePacket.Result == XR_SUCCESS);
+
+	if (BasePacket.Result != XR_SUCCESS)
+	{
+		return false;
+	}
 
 	FOpenXRWaitSwapchainImagePacket Data(XrResult::XR_ERROR_RUNTIME_FAILURE);
 
 	EncodedData << Data;
+	return true;
 
 	// nothing to do yet
 }
-void FOpenXRCaptureDecoder::DecodeReleaseSwapchainImage(const FOpenXRAPIPacketBase& BasePacket)
+bool FOpenXRCaptureDecoder::DecodeReleaseSwapchainImage(const FOpenXRAPIPacketBase& BasePacket)
 {
 	check(BasePacket.ApiId == EOpenXRAPIPacketId::ReleaseSwapchainImage);
-	check(BasePacket.Result == XR_SUCCESS);
+
+	if (BasePacket.Result != XR_SUCCESS)
+	{
+		return false;
+	}
 
 	FOpenXRReleaseSwapchainImagePacket Data(XrResult::XR_ERROR_RUNTIME_FAILURE);
 
 	EncodedData << Data;
+	return true;
 
 	// nothing to do yet
 }
-void FOpenXRCaptureDecoder::DecodeBeginSession(const FOpenXRAPIPacketBase& BasePacket)
+bool FOpenXRCaptureDecoder::DecodeBeginSession(const FOpenXRAPIPacketBase& BasePacket)
 {
 	check(BasePacket.ApiId == EOpenXRAPIPacketId::BeginSession);
-	check(BasePacket.Result == XR_SUCCESS);
+
+	if (BasePacket.Result != XR_SUCCESS)
+	{
+		return false;
+	}
 
 	FOpenXRBeginSessionPacket Data(XrResult::XR_ERROR_RUNTIME_FAILURE);
 
 	EncodedData << Data;
+	return true;
+
 	// nothing to do yet
 }
 
-void FOpenXRCaptureDecoder::DecodeEndSession(const FOpenXRAPIPacketBase& BasePacket)
+bool FOpenXRCaptureDecoder::DecodeEndSession(const FOpenXRAPIPacketBase& BasePacket)
 {
 	check(BasePacket.ApiId == EOpenXRAPIPacketId::EndSession);
-	check(BasePacket.Result == XR_SUCCESS);
+
+	if (BasePacket.Result != XR_SUCCESS)
+	{
+		return false;
+	}
 
 	FOpenXREndSessionPacket Data(XrResult::XR_ERROR_RUNTIME_FAILURE);
 
 	EncodedData << Data;
+	return true;
+
 	// nothing to do yet
 }
 
-void FOpenXRCaptureDecoder::DecodeRequestExitSession(const FOpenXRAPIPacketBase& BasePacket)
+bool FOpenXRCaptureDecoder::DecodeRequestExitSession(const FOpenXRAPIPacketBase& BasePacket)
 {
 	check(BasePacket.ApiId == EOpenXRAPIPacketId::RequestExitSession);
-	check(BasePacket.Result == XR_SUCCESS);
+
+	if (BasePacket.Result != XR_SUCCESS)
+	{
+		return false;
+	}
 
 	FOpenXRRequestExitSessionPacket Data(XrResult::XR_ERROR_RUNTIME_FAILURE);
 
 	EncodedData << Data;
+	return true;
+
 	// nothing to do yet
 }
 
-void FOpenXRCaptureDecoder::DecodeWaitFrame(const FOpenXRAPIPacketBase& BasePacket)
+bool FOpenXRCaptureDecoder::DecodeWaitFrame(const FOpenXRAPIPacketBase& BasePacket)
 {
 	check(BasePacket.ApiId == EOpenXRAPIPacketId::WaitFrame);
-	check(BasePacket.Result == XR_SUCCESS);
+
+	if (BasePacket.Result != XR_SUCCESS)
+	{
+		return false;
+	}
 
 	FOpenXRWaitFramePacket Data(XrResult::XR_ERROR_RUNTIME_FAILURE);
 
 	EncodedData << Data;
 
 	WaitFrames.Add(Data);
+	return true;
 }
 
-void FOpenXRCaptureDecoder::DecodeBeginFrame(const FOpenXRAPIPacketBase& BasePacket)
+bool FOpenXRCaptureDecoder::DecodeBeginFrame(const FOpenXRAPIPacketBase& BasePacket)
 {
 	check(BasePacket.ApiId == EOpenXRAPIPacketId::BeginFrame);
-	check(BasePacket.Result == XR_SUCCESS || BasePacket.Result == XR_FRAME_DISCARDED);
+
+	// For xrBeginFrame, XR_FRAME_DISCARDED is also considered a successful result
+	if (BasePacket.Result != XR_SUCCESS && BasePacket.Result != XR_FRAME_DISCARDED)
+	{
+		return false;
+	}
 
 	FOpenXRBeginFramePacket Data(XrResult::XR_ERROR_RUNTIME_FAILURE);
 
 	EncodedData << Data;
+	return true;
+
 	// nothing to do yet
 }
 
-void FOpenXRCaptureDecoder::DecodeEndFrame(const FOpenXRAPIPacketBase& BasePacket)
+bool FOpenXRCaptureDecoder::DecodeEndFrame(const FOpenXRAPIPacketBase& BasePacket)
 {
 	check(BasePacket.ApiId == EOpenXRAPIPacketId::EndFrame);
-	check(BasePacket.Result == XR_SUCCESS);
+
+	if (BasePacket.Result != XR_SUCCESS)
+	{
+		return false;
+	}
 
 	FOpenXREndFramePacket Data(XrResult::XR_ERROR_RUNTIME_FAILURE);
 
 	EncodedData << Data;
+	return true;
 
 	// Not really sure what to do with this wealth of info here!
 }
 
-void FOpenXRCaptureDecoder::DecodeLocateViews(const FOpenXRAPIPacketBase& BasePacket)
+bool FOpenXRCaptureDecoder::DecodeLocateViews(const FOpenXRAPIPacketBase& BasePacket)
 {
 	check(BasePacket.ApiId == EOpenXRAPIPacketId::LocateViews);
-	check(BasePacket.Result == XR_SUCCESS);
+
+	if (BasePacket.Result != XR_SUCCESS)
+	{
+		return false;
+	}
 
 	FOpenXRLocateViewsPacket Data(XrResult::XR_ERROR_RUNTIME_FAILURE);
 
@@ -559,82 +757,122 @@ void FOpenXRCaptureDecoder::DecodeLocateViews(const FOpenXRAPIPacketBase& BasePa
 	{
 		ViewLocations[Data.ViewLocateInfo.viewConfigurationType].Add(Data);
 	}
+
+	return true;
 }
 
-void FOpenXRCaptureDecoder::DecodeStringToPath(const FOpenXRAPIPacketBase& BasePacket)
+bool FOpenXRCaptureDecoder::DecodeStringToPath(const FOpenXRAPIPacketBase& BasePacket)
 {
 	check(BasePacket.ApiId == EOpenXRAPIPacketId::StringToPath);
-	check(BasePacket.Result == XR_SUCCESS);
+
+	if (BasePacket.Result != XR_SUCCESS)
+	{
+		return false;
+	}
 
 	FOpenXRStringToPathPacket Data(XrResult::XR_ERROR_RUNTIME_FAILURE);
 
 	EncodedData << Data;
 
 	PathToStringMap.Add(Data.GeneratedPath, FName(ANSI_TO_TCHAR(Data.PathStringToWrite.GetData())));
+	return true;
 
 	// TODO: Do we need a bi-directional map at any point?
 }
 
-void FOpenXRCaptureDecoder::DecodePathToString(const FOpenXRAPIPacketBase& BasePacket)
+bool FOpenXRCaptureDecoder::DecodePathToString(const FOpenXRAPIPacketBase& BasePacket)
 {
 	check(BasePacket.ApiId == EOpenXRAPIPacketId::PathToString);
-	check(BasePacket.Result == XR_SUCCESS);
+
+	if (BasePacket.Result != XR_SUCCESS)
+	{
+		return false;
+	}
 
 	FOpenXRPathToStringPacket Data(XrResult::XR_ERROR_RUNTIME_FAILURE);
 
 	EncodedData << Data;
+	return true;
+
 	// nothing to do yet
 }
 
-void FOpenXRCaptureDecoder::DecodeCreateActionSet(const FOpenXRAPIPacketBase& BasePacket)
+bool FOpenXRCaptureDecoder::DecodeCreateActionSet(const FOpenXRAPIPacketBase& BasePacket)
 {
 	check(BasePacket.ApiId == EOpenXRAPIPacketId::CreateActionSet);
-	check(BasePacket.Result == XR_SUCCESS);
+
+	if (BasePacket.Result != XR_SUCCESS)
+	{
+		return false;
+	}
 
 	FOpenXRCreateActionSetPacket Data(XrResult::XR_ERROR_RUNTIME_FAILURE);
 
 	EncodedData << Data;
+	return true;
+
 	// nothing to do yet
 }
 
-void FOpenXRCaptureDecoder::DecodeDestroyActionSet(const FOpenXRAPIPacketBase& BasePacket)
+bool FOpenXRCaptureDecoder::DecodeDestroyActionSet(const FOpenXRAPIPacketBase& BasePacket)
 {
 	check(BasePacket.ApiId == EOpenXRAPIPacketId::DestroyActionSet);
-	check(BasePacket.Result == XR_SUCCESS);
+
+	if (BasePacket.Result != XR_SUCCESS)
+	{
+		return false;
+	}
 
 	FOpenXRDestroyActionSetPacket Data(XrResult::XR_ERROR_RUNTIME_FAILURE);
 
 	EncodedData << Data;
+	return true;
+
 	// nothing to do yet
 }
 
-void FOpenXRCaptureDecoder::DecodeCreateAction(const FOpenXRAPIPacketBase& BasePacket)
+bool FOpenXRCaptureDecoder::DecodeCreateAction(const FOpenXRAPIPacketBase& BasePacket)
 {
 	check(BasePacket.ApiId == EOpenXRAPIPacketId::CreateAction);
-	check(BasePacket.Result == XR_SUCCESS);
+
+	if (BasePacket.Result != XR_SUCCESS)
+	{
+		return false;
+	}
 
 	FOpenXRCreateActionPacket Data(XrResult::XR_ERROR_RUNTIME_FAILURE);
 
 	EncodedData << Data;
 
 	CreatedActions.Add(Data);
+	return true;
 }
 
-void FOpenXRCaptureDecoder::DecodeDestroyAction(const FOpenXRAPIPacketBase& BasePacket)
+bool FOpenXRCaptureDecoder::DecodeDestroyAction(const FOpenXRAPIPacketBase& BasePacket)
 {
 	check(BasePacket.ApiId == EOpenXRAPIPacketId::DestroyAction);
-	check(BasePacket.Result == XR_SUCCESS);
+
+	if (BasePacket.Result != XR_SUCCESS)
+	{
+		return false;
+	}
 
 	FOpenXRDestroyActionPacket Data(XrResult::XR_ERROR_RUNTIME_FAILURE);
 
 	EncodedData << Data;
+	return true;
+
 	// nothing to do yet
 }
 
-void FOpenXRCaptureDecoder::DecodeSuggestInteractionProfileBindings(const FOpenXRAPIPacketBase& BasePacket)
+bool FOpenXRCaptureDecoder::DecodeSuggestInteractionProfileBindings(const FOpenXRAPIPacketBase& BasePacket)
 {
 	check(BasePacket.ApiId == EOpenXRAPIPacketId::SuggestInteractionProfileBindings);
-	check(BasePacket.Result == XR_SUCCESS);
+
+	if (BasePacket.Result != XR_SUCCESS)
+	{
+		return false;
+	}
 
 	FOpenXRSuggestInteractionProfileBindingsPacket Data(XrResult::XR_ERROR_RUNTIME_FAILURE);
 
@@ -642,34 +880,52 @@ void FOpenXRCaptureDecoder::DecodeSuggestInteractionProfileBindings(const FOpenX
 
 	check(PathToStringMap.Contains(Data.InteractionProfile));
 	StringToSuggestedBindingsMap.Add(PathToStringMap[Data.InteractionProfile], Data.SuggestedBindings);
+
+	return true;
 }
 
-void FOpenXRCaptureDecoder::DecodeAttachSessionActionSets(const FOpenXRAPIPacketBase& BasePacket)
+bool FOpenXRCaptureDecoder::DecodeAttachSessionActionSets(const FOpenXRAPIPacketBase& BasePacket)
 {
 	check(BasePacket.ApiId == EOpenXRAPIPacketId::AttachSessionActionSets);
-	check(BasePacket.Result == XR_SUCCESS);
+
+	if (BasePacket.Result != XR_SUCCESS)
+	{
+		return false;
+	}
 
 	FOpenXRAttachSessionActionSetsPacket Data(XrResult::XR_ERROR_RUNTIME_FAILURE);
 
 	EncodedData << Data;
+	return true;
+
 	// nothing to do yet
 }
 
-void FOpenXRCaptureDecoder::DecodeGetCurrentInteractionProfile(const FOpenXRAPIPacketBase& BasePacket)
+bool FOpenXRCaptureDecoder::DecodeGetCurrentInteractionProfile(const FOpenXRAPIPacketBase& BasePacket)
 {
 	check(BasePacket.ApiId == EOpenXRAPIPacketId::GetCurrentInteractionProfile);
-	check(BasePacket.Result == XR_SUCCESS);
+
+	if (BasePacket.Result != XR_SUCCESS)
+	{
+		return false;
+	}
 
 	FOpenXRGetCurrentInteractionProfilePacket Data(XrResult::XR_ERROR_RUNTIME_FAILURE);
 
 	EncodedData << Data;
+	return true;
+
 	// nothing to do yet
 }
 
-void FOpenXRCaptureDecoder::DecodeGetActionStateBoolean(const FOpenXRAPIPacketBase& BasePacket)
+bool FOpenXRCaptureDecoder::DecodeGetActionStateBoolean(const FOpenXRAPIPacketBase& BasePacket)
 {
 	check(BasePacket.ApiId == EOpenXRAPIPacketId::GetActionStateBoolean);
-	check(BasePacket.Result == XR_SUCCESS);
+
+	if (BasePacket.Result != XR_SUCCESS)
+	{
+		return false;
+	}
 
 	FOpenXRGetActionStateBooleanPacket Data(XrResult::XR_ERROR_RUNTIME_FAILURE);
 
@@ -680,12 +936,18 @@ void FOpenXRCaptureDecoder::DecodeGetActionStateBoolean(const FOpenXRAPIPacketBa
 		BooleanActionStates.Add(Data.GetInfoBoolean.action);
 	}
 	BooleanActionStates[Data.GetInfoBoolean.action].Add(Data);
+
+	return true;
 }
 
-void FOpenXRCaptureDecoder::DecodeGetActionStateFloat(const FOpenXRAPIPacketBase& BasePacket)
+bool FOpenXRCaptureDecoder::DecodeGetActionStateFloat(const FOpenXRAPIPacketBase& BasePacket)
 {
 	check(BasePacket.ApiId == EOpenXRAPIPacketId::GetActionStateFloat);
-	check(BasePacket.Result == XR_SUCCESS);
+
+	if (BasePacket.Result != XR_SUCCESS)
+	{
+		return false;
+	}
 
 	FOpenXRGetActionStateFloatPacket Data(XrResult::XR_ERROR_RUNTIME_FAILURE);
 
@@ -696,12 +958,18 @@ void FOpenXRCaptureDecoder::DecodeGetActionStateFloat(const FOpenXRAPIPacketBase
 		FloatActionStates.Add(Data.GetInfoFloat.action);
 	}
 	FloatActionStates[Data.GetInfoFloat.action].Add(Data);
+
+	return true;
 }
 
-void FOpenXRCaptureDecoder::DecodeGetActionStateVector2f(const FOpenXRAPIPacketBase& BasePacket)
+bool FOpenXRCaptureDecoder::DecodeGetActionStateVector2f(const FOpenXRAPIPacketBase& BasePacket)
 {
 	check(BasePacket.ApiId == EOpenXRAPIPacketId::GetActionStateVector2F);
-	check(BasePacket.Result == XR_SUCCESS);
+
+	if (BasePacket.Result != XR_SUCCESS)
+	{
+		return false;
+	}
 
 	FOpenXRGetActionStateVector2fPacket Data(XrResult::XR_ERROR_RUNTIME_FAILURE);
 
@@ -712,12 +980,18 @@ void FOpenXRCaptureDecoder::DecodeGetActionStateVector2f(const FOpenXRAPIPacketB
 		VectorActionStates.Add(Data.GetInfoVector2f.action);
 	}
 	VectorActionStates[Data.GetInfoVector2f.action].Add(Data);
+
+	return true;
 }
 
-void FOpenXRCaptureDecoder::DecodeGetActionStatePose(const FOpenXRAPIPacketBase& BasePacket)
+bool FOpenXRCaptureDecoder::DecodeGetActionStatePose(const FOpenXRAPIPacketBase& BasePacket)
 {
 	check(BasePacket.ApiId == EOpenXRAPIPacketId::GetActionStatePose);
-	check(BasePacket.Result == XR_SUCCESS);
+
+	if (BasePacket.Result != XR_SUCCESS)
+	{
+		return false;
+	}
 
 	FOpenXRGetActionStatePosePacket Data(XrResult::XR_ERROR_RUNTIME_FAILURE);
 
@@ -728,92 +1002,128 @@ void FOpenXRCaptureDecoder::DecodeGetActionStatePose(const FOpenXRAPIPacketBase&
 		PoseActionStates.Add(Data.GetInfoPose.action);
 	}
 	PoseActionStates[Data.GetInfoPose.action].Add(Data);
+
+	return true;
 }
 
-void FOpenXRCaptureDecoder::DecodeSyncActions(const FOpenXRAPIPacketBase& BasePacket)
+bool FOpenXRCaptureDecoder::DecodeSyncActions(const FOpenXRAPIPacketBase& BasePacket)
 {
 	check(BasePacket.ApiId == EOpenXRAPIPacketId::SyncActions);
 
-	// We can succeed here with NOT_FOCUSED
-	check((BasePacket.Result == XR_SUCCESS) || (BasePacket.Result == XR_SESSION_NOT_FOCUSED));
+	// For xrSyncActions, XR_SESSION_NOT_FOCUSED is also considered a successful result
+	if (BasePacket.Result != XR_SUCCESS && BasePacket.Result != XR_SESSION_NOT_FOCUSED)
+	{
+		return false;
+	}
 
 	FOpenXRSyncActionsPacket Data(XrResult::XR_ERROR_RUNTIME_FAILURE);
 
 	EncodedData << Data;
 
 	SyncActions.Add(Data);
+	return true;
 }
 
-//void FOpenXRCaptureDecoder::DecodeEnumerateBoundSourcesForAction(const FOpenXRAPIPacketBase& BasePacket)
+//bool FOpenXRCaptureDecoder::DecodeEnumerateBoundSourcesForAction(const FOpenXRAPIPacketBase& BasePacket)
 //{
 //
 //}
-//void FOpenXRCaptureDecoder::DecodeGetInputSourceLocalizedName(const FOpenXRAPIPacketBase& BasePacket)
+//bool FOpenXRCaptureDecoder::DecodeGetInputSourceLocalizedName(const FOpenXRAPIPacketBase& BasePacket)
 //{
 //
 //}
 
-void FOpenXRCaptureDecoder::DecodeApplyHapticFeedback(const FOpenXRAPIPacketBase& BasePacket)
+bool FOpenXRCaptureDecoder::DecodeApplyHapticFeedback(const FOpenXRAPIPacketBase& BasePacket)
 {
 	check(BasePacket.ApiId == EOpenXRAPIPacketId::ApplyHapticFeedback);
-	check(BasePacket.Result == XR_SUCCESS);
+
+	if (BasePacket.Result != XR_SUCCESS)
+	{
+		return false;
+	}
 
 	FOpenXRApplyHapticFeedbackPacket Data(XrResult::XR_ERROR_RUNTIME_FAILURE);
 
 	EncodedData << Data;
+	return true;
 }
 
-void FOpenXRCaptureDecoder::DecodeStopHapticFeedback(const FOpenXRAPIPacketBase& BasePacket)
+bool FOpenXRCaptureDecoder::DecodeStopHapticFeedback(const FOpenXRAPIPacketBase& BasePacket)
 {
 	check(BasePacket.ApiId == EOpenXRAPIPacketId::StopHapticFeedback);
-	check(BasePacket.Result == XR_SUCCESS);
+
+	if (BasePacket.Result != XR_SUCCESS)
+	{
+		return false;
+	}
 
 	FOpenXRStopHapticFeedbackPacket Data(XrResult::XR_ERROR_RUNTIME_FAILURE);
 
 	EncodedData << Data;
+	return true;
 }
 
-void FOpenXRCaptureDecoder::DecodeInitializeLoaderKHR(const FOpenXRAPIPacketBase& BasePacket)
+bool FOpenXRCaptureDecoder::DecodeInitializeLoaderKHR(const FOpenXRAPIPacketBase& BasePacket)
 {
 	check(BasePacket.ApiId == EOpenXRAPIPacketId::InitializeLoaderKHR);
-	check(BasePacket.Result == XR_SUCCESS);
+
+	if (BasePacket.Result != XR_SUCCESS)
+	{
+		return false;
+	}
 
 	FOpenXRInitializeLoaderKHRPacket Data(XrResult::XR_ERROR_RUNTIME_FAILURE);
 
 	EncodedData << Data;
+	return true;
 }
 
-void FOpenXRCaptureDecoder::DecodeGetVisibilityMaskKHR(const FOpenXRAPIPacketBase& BasePacket)
+bool FOpenXRCaptureDecoder::DecodeGetVisibilityMaskKHR(const FOpenXRAPIPacketBase& BasePacket)
 {
 	check(BasePacket.ApiId == EOpenXRAPIPacketId::GetVisibilityMaskKHR);
-	check(BasePacket.Result == XR_SUCCESS);
+
+	if (BasePacket.Result != XR_SUCCESS)
+	{
+		return false;
+	}
 
 	FOpenXRGetVisibilityMaskKHRPacket Data(XrResult::XR_ERROR_RUNTIME_FAILURE);
 
 	EncodedData << Data;
+	return true;
 }
 
 #if defined(XR_USE_GRAPHICS_API_D3D11)
-void FOpenXRCaptureDecoder::DecodeGetD3D11GraphicsRequirementsKHR(const FOpenXRAPIPacketBase& BasePacket)
+bool FOpenXRCaptureDecoder::DecodeGetD3D11GraphicsRequirementsKHR(const FOpenXRAPIPacketBase& BasePacket)
 {
 	check(BasePacket.ApiId == EOpenXRAPIPacketId::GetD3D11GraphicsRequirementsKHR);
-	check(BasePacket.Result == XR_SUCCESS);
+
+	if (BasePacket.Result != XR_SUCCESS)
+	{
+		return false;
+	}
 
 	FOpenXRGetD3D11GraphicsRequirementsKHRPacket Data(XrResult::XR_ERROR_RUNTIME_FAILURE);
 
 	EncodedData << Data;
+	return true;
 }
 #endif
 
 #if defined(XR_USE_GRAPHICS_API_D3D12)
-void FOpenXRCaptureDecoder::DecodeGetD3D12GraphicsRequirementsKHR(const FOpenXRAPIPacketBase& BasePacket)
+bool FOpenXRCaptureDecoder::DecodeGetD3D12GraphicsRequirementsKHR(const FOpenXRAPIPacketBase& BasePacket)
 {
 	check(BasePacket.ApiId == EOpenXRAPIPacketId::GetD3D12GraphicsRequirementsKHR);
-	check(BasePacket.Result == XR_SUCCESS);
+
+	if (BasePacket.Result != XR_SUCCESS)
+	{
+		return false;
+	}
 
 	FOpenXRGetD3D12GraphicsRequirementsKHRPacket Data(XrResult::XR_ERROR_RUNTIME_FAILURE);
 
 	EncodedData << Data;
+	return true;
 }
 #endif
 
diff --git a/Engine/Plugins/Runtime/XR/XRScribe/Source/XRScribe/Private/XRScribeEmulationLayer.cpp b/Engine/Plugins/Runtime/XR/XRScribe/Source/XRScribe/Private/XRScribeEmulationLayer.cpp
--- a/Engine/Plugins/Runtime/XR/XRScribe/Source/XRScribe/Private/XRScribeEmulationLayer.cpp
+++ b/Engine/Plugins/Runtime/XR/XRScribe/Source/XRScribe/Private/XRScribeEmulationLayer.cpp
@@ -13,10 +13,9 @@
 #include "OpenXRCore.h"
 #include "RHICommandList.h"
 #include "RHIResources.h"
+#include "Misc/MessageDialog.h"
 
 
-DEFINE_LOG_CATEGORY_STATIC(LogXRScribeEmulate, Log, All);
-
 // TODO
 // * useful error logs at fail points
 // * per-api next pointer validation
@@ -268,27 +267,51 @@ bool FOpenXREmulationLayer::LoadCaptureFromFile(const FString& EmulationLoadPath
 	if (FFileHelper::LoadFileToArray(CaptureDecoder.GetEncodedData(), *EmulationLoadPath))
 	{
 		UE_LOG(LogXRScribeEmulate, Log, TEXT("Capture successfully loaded: %s"), *EmulationLoadPath);
-		CaptureDecoder.DecodeDataFromMemory();
+		if (CaptureDecoder.DecodeDataFromMemory())
+		{
+			PostLoadActions();
+			return true;
+		}
+		else
+		{
+			UE_LOG(LogXRScribeEmulate, Error, TEXT("Failed to decode capture"));
+
+			// On Windows, a workaround in FOpenXRHMDModule::PreInit means we may end up initializing XRScribe twice
+			// For now we'll always load the same file, if it's invalid we should only display an error message once
+			static bool bDisplayedErrorMessage = false;
+
+			if (!bDisplayedErrorMessage)
+			{
+				FMessageDialog::Open(EAppMsgCategory::Error, EAppMsgType::Ok, FText::FromString(FString("XRScribe capture at ") + EmulationLoadPath + FString(" was malformed or captured with a different version of the plugin. Emulation layer will be disabled.\n\nSee log for details.")));
+				bDisplayedErrorMessage = true;
+			}
+			
+			// Clean up loaded data
+			CaptureDecoder.GetEncodedData().Empty();
+			return false;
+		}
 	}
 	else
 	{
 		UE_LOG(LogXRScribeEmulate, Error, TEXT("Capture failed to load"));
 		return false;
 	}
-
-	PostLoadActions();
-
-	return true;
 }
 
 bool FOpenXREmulationLayer::LoadCaptureFromData(const TArray<uint8>& EncodedData)
 {
 	CaptureDecoder.GetEncodedData().Append(EncodedData);
-	CaptureDecoder.DecodeDataFromMemory();
-
-	PostLoadActions();
 
-	return true;
+	if (CaptureDecoder.DecodeDataFromMemory())
+	{
+		PostLoadActions();
+		return true;
+	}
+	else
+	{
+		CaptureDecoder.GetEncodedData().Empty();
+		return false;
+	}
 }
 
 void FOpenXREmulationLayer::PostLoadActions()
diff --git a/Engine/Plugins/Runtime/XR/XRScribe/Source/XRScribe/Public/XRScribeAPIDecoder.h b/Engine/Plugins/Runtime/XR/XRScribe/Source/XRScribe/Public/XRScribeAPIDecoder.h
--- a/Engine/Plugins/Runtime/XR/XRScribe/Source/XRScribe/Public/XRScribeAPIDecoder.h
+++ b/Engine/Plugins/Runtime/XR/XRScribe/Source/XRScribe/Public/XRScribeAPIDecoder.h
@@ -8,6 +8,9 @@
 #include "Serialization/ArrayReader.h"
 #include "XRScribeFileFormat.h"
 
+
+DEFINE_LOG_CATEGORY_STATIC(LogXRScribeEmulate, Log, All);
+
 namespace UE::XRScribe
 {
 
@@ -17,7 +20,7 @@ public:
 	explicit FOpenXRCaptureDecoder();
 	~FOpenXRCaptureDecoder();
 
-	void DecodeDataFromMemory();
+	bool DecodeDataFromMemory();
 
 	// state accessors
 	[[nodiscard]] const TArray<XrExtensionProperties>& GetInstanceExtensionProperties() { return InstanceExtensionProperties; }
@@ -56,74 +59,74 @@ public:
 protected:
 
 	// packet decoders
-	void DecodeEnumerateApiLayerProperties(const FOpenXRAPIPacketBase& BasePacket);
-	void DecodeEnumerateInstanceExtensionProperties(const FOpenXRAPIPacketBase& BasePacket);
-	void DecodeCreateInstance(const FOpenXRAPIPacketBase& BasePacket);
-	void DecodeDestroyInstance(const FOpenXRAPIPacketBase& BasePacket);
-	void DecodeGetInstanceProperties(const FOpenXRAPIPacketBase& BasePacket);
-	//void DecodePollEvent(const FOpenXRAPIPacketBase& BasePacket);
-	//void DecodeResultToString(const FOpenXRAPIPacketBase& BasePacket);
-	//void DecodeStructureTypeToString(const FOpenXRAPIPacketBase& BasePacket);
-	void DecodeGetSystem(const FOpenXRAPIPacketBase& BasePacket);
-	void DecodeGetSystemProperties(const FOpenXRAPIPacketBase& BasePacket);
-	void DecodeEnumerateEnvironmentBlendModes(const FOpenXRAPIPacketBase& BasePacket);
-	void DecodeCreateSession(const FOpenXRAPIPacketBase& BasePacket);
-	void DecodeDestroySession(const FOpenXRAPIPacketBase& BasePacket);
-	void DecodeEnumerateReferenceSpaces(const FOpenXRAPIPacketBase& BasePacket);
-	void DecodeCreateReferenceSpace(const FOpenXRAPIPacketBase& BasePacket);
-	void DecodeGetReferenceSpaceBoundsRect(const FOpenXRAPIPacketBase& BasePacket);
-	void DecodeCreateActionSpace(const FOpenXRAPIPacketBase& BasePacket);
-	void DecodeLocateSpace(const FOpenXRAPIPacketBase& BasePacket);
-	void DecodeDestroySpace(const FOpenXRAPIPacketBase& BasePacket);
-	void DecodeEnumerateViewConfigurations(const FOpenXRAPIPacketBase& BasePacket);
-	void DecodeGetViewConfigurationProperties(const FOpenXRAPIPacketBase& BasePacket);
-	void DecodeEnumerateViewConfigurationViews(const FOpenXRAPIPacketBase& BasePacket);
-	void DecodeEnumerateSwapchainFormats(const FOpenXRAPIPacketBase& BasePacket);
-	void DecodeCreateSwapchain(const FOpenXRAPIPacketBase& BasePacket);
-	void DecodeDestroySwapchain(const FOpenXRAPIPacketBase& BasePacket);
-	void DecodeEnumerateSwapchainImages(const FOpenXRAPIPacketBase& BasePacket);
-	void DecodeAcquireSwapchainImage(const FOpenXRAPIPacketBase& BasePacket);
-	void DecodeWaitSwapchainImage(const FOpenXRAPIPacketBase& BasePacket);
-	void DecodeReleaseSwapchainImage(const FOpenXRAPIPacketBase& BasePacket);
-	void DecodeBeginSession(const FOpenXRAPIPacketBase& BasePacket);
-	void DecodeEndSession(const FOpenXRAPIPacketBase& BasePacket);
-	void DecodeRequestExitSession(const FOpenXRAPIPacketBase& BasePacket);
-	void DecodeWaitFrame(const FOpenXRAPIPacketBase& BasePacket);
-	void DecodeBeginFrame(const FOpenXRAPIPacketBase& BasePacket);
-	void DecodeEndFrame(const FOpenXRAPIPacketBase& BasePacket);
-	void DecodeLocateViews(const FOpenXRAPIPacketBase& BasePacket);
-	void DecodeStringToPath(const FOpenXRAPIPacketBase& BasePacket);
-	void DecodePathToString(const FOpenXRAPIPacketBase& BasePacket);
-	void DecodeCreateActionSet(const FOpenXRAPIPacketBase& BasePacket);
-	void DecodeDestroyActionSet(const FOpenXRAPIPacketBase& BasePacket);
-	void DecodeCreateAction(const FOpenXRAPIPacketBase& BasePacket);
-	void DecodeDestroyAction(const FOpenXRAPIPacketBase& BasePacket);
-
-	void DecodeSuggestInteractionProfileBindings(const FOpenXRAPIPacketBase& BasePacket);
-	void DecodeAttachSessionActionSets(const FOpenXRAPIPacketBase& BasePacket);
-	void DecodeGetCurrentInteractionProfile(const FOpenXRAPIPacketBase& BasePacket);
-	void DecodeGetActionStateBoolean(const FOpenXRAPIPacketBase& BasePacket);
-	void DecodeGetActionStateFloat(const FOpenXRAPIPacketBase& BasePacket);
-	void DecodeGetActionStateVector2f(const FOpenXRAPIPacketBase& BasePacket);
-	void DecodeGetActionStatePose(const FOpenXRAPIPacketBase& BasePacket);
-	void DecodeSyncActions(const FOpenXRAPIPacketBase& BasePacket);
-	//void DecodeEnumerateBoundSourcesForAction(const FOpenXRAPIPacketBase& BasePacket);
-	//void DecodeGetInputSourceLocalizedName(const FOpenXRAPIPacketBase& BasePacket);
-
-	void DecodeApplyHapticFeedback(const FOpenXRAPIPacketBase& BasePacket);
-	void DecodeStopHapticFeedback(const FOpenXRAPIPacketBase& BasePacket);
-	void DecodeInitializeLoaderKHR(const FOpenXRAPIPacketBase& BasePacket);
-	void DecodeGetVisibilityMaskKHR(const FOpenXRAPIPacketBase& BasePacket);
+	bool DecodeEnumerateApiLayerProperties(const FOpenXRAPIPacketBase& BasePacket);
+	bool DecodeEnumerateInstanceExtensionProperties(const FOpenXRAPIPacketBase& BasePacket);
+	bool DecodeCreateInstance(const FOpenXRAPIPacketBase& BasePacket);
+	bool DecodeDestroyInstance(const FOpenXRAPIPacketBase& BasePacket);
+	bool DecodeGetInstanceProperties(const FOpenXRAPIPacketBase& BasePacket);
+	//bool DecodePollEvent(const FOpenXRAPIPacketBase& BasePacket);
+	//bool DecodeResultToString(const FOpenXRAPIPacketBase& BasePacket);
+	//bool DecodeStructureTypeToString(const FOpenXRAPIPacketBase& BasePacket);
+	bool DecodeGetSystem(const FOpenXRAPIPacketBase& BasePacket);
+	bool DecodeGetSystemProperties(const FOpenXRAPIPacketBase& BasePacket);
+	bool DecodeEnumerateEnvironmentBlendModes(const FOpenXRAPIPacketBase& BasePacket);
+	bool DecodeCreateSession(const FOpenXRAPIPacketBase& BasePacket);
+	bool DecodeDestroySession(const FOpenXRAPIPacketBase& BasePacket);
+	bool DecodeEnumerateReferenceSpaces(const FOpenXRAPIPacketBase& BasePacket);
+	bool DecodeCreateReferenceSpace(const FOpenXRAPIPacketBase& BasePacket);
+	bool DecodeGetReferenceSpaceBoundsRect(const FOpenXRAPIPacketBase& BasePacket);
+	bool DecodeCreateActionSpace(const FOpenXRAPIPacketBase& BasePacket);
+	bool DecodeLocateSpace(const FOpenXRAPIPacketBase& BasePacket);
+	bool DecodeDestroySpace(const FOpenXRAPIPacketBase& BasePacket);
+	bool DecodeEnumerateViewConfigurations(const FOpenXRAPIPacketBase& BasePacket);
+	bool DecodeGetViewConfigurationProperties(const FOpenXRAPIPacketBase& BasePacket);
+	bool DecodeEnumerateViewConfigurationViews(const FOpenXRAPIPacketBase& BasePacket);
+	bool DecodeEnumerateSwapchainFormats(const FOpenXRAPIPacketBase& BasePacket);
+	bool DecodeCreateSwapchain(const FOpenXRAPIPacketBase& BasePacket);
+	bool DecodeDestroySwapchain(const FOpenXRAPIPacketBase& BasePacket);
+	bool DecodeEnumerateSwapchainImages(const FOpenXRAPIPacketBase& BasePacket);
+	bool DecodeAcquireSwapchainImage(const FOpenXRAPIPacketBase& BasePacket);
+	bool DecodeWaitSwapchainImage(const FOpenXRAPIPacketBase& BasePacket);
+	bool DecodeReleaseSwapchainImage(const FOpenXRAPIPacketBase& BasePacket);
+	bool DecodeBeginSession(const FOpenXRAPIPacketBase& BasePacket);
+	bool DecodeEndSession(const FOpenXRAPIPacketBase& BasePacket);
+	bool DecodeRequestExitSession(const FOpenXRAPIPacketBase& BasePacket);
+	bool DecodeWaitFrame(const FOpenXRAPIPacketBase& BasePacket);
+	bool DecodeBeginFrame(const FOpenXRAPIPacketBase& BasePacket);
+	bool DecodeEndFrame(const FOpenXRAPIPacketBase& BasePacket);
+	bool DecodeLocateViews(const FOpenXRAPIPacketBase& BasePacket);
+	bool DecodeStringToPath(const FOpenXRAPIPacketBase& BasePacket);
+	bool DecodePathToString(const FOpenXRAPIPacketBase& BasePacket);
+	bool DecodeCreateActionSet(const FOpenXRAPIPacketBase& BasePacket);
+	bool DecodeDestroyActionSet(const FOpenXRAPIPacketBase& BasePacket);
+	bool DecodeCreateAction(const FOpenXRAPIPacketBase& BasePacket);
+	bool DecodeDestroyAction(const FOpenXRAPIPacketBase& BasePacket);
+
+	bool DecodeSuggestInteractionProfileBindings(const FOpenXRAPIPacketBase& BasePacket);
+	bool DecodeAttachSessionActionSets(const FOpenXRAPIPacketBase& BasePacket);
+	bool DecodeGetCurrentInteractionProfile(const FOpenXRAPIPacketBase& BasePacket);
+	bool DecodeGetActionStateBoolean(const FOpenXRAPIPacketBase& BasePacket);
+	bool DecodeGetActionStateFloat(const FOpenXRAPIPacketBase& BasePacket);
+	bool DecodeGetActionStateVector2f(const FOpenXRAPIPacketBase& BasePacket);
+	bool DecodeGetActionStatePose(const FOpenXRAPIPacketBase& BasePacket);
+	bool DecodeSyncActions(const FOpenXRAPIPacketBase& BasePacket);
+	//bool DecodeEnumerateBoundSourcesForAction(const FOpenXRAPIPacketBase& BasePacket);
+	//bool DecodeGetInputSourceLocalizedName(const FOpenXRAPIPacketBase& BasePacket);
+
+	bool DecodeApplyHapticFeedback(const FOpenXRAPIPacketBase& BasePacket);
+	bool DecodeStopHapticFeedback(const FOpenXRAPIPacketBase& BasePacket);
+	bool DecodeInitializeLoaderKHR(const FOpenXRAPIPacketBase& BasePacket);
+	bool DecodeGetVisibilityMaskKHR(const FOpenXRAPIPacketBase& BasePacket);
 #if defined(XR_USE_GRAPHICS_API_D3D11)
-	void DecodeGetD3D11GraphicsRequirementsKHR(const FOpenXRAPIPacketBase& BasePacket);
+	bool DecodeGetD3D11GraphicsRequirementsKHR(const FOpenXRAPIPacketBase& BasePacket);
 #endif
 #if defined(XR_USE_GRAPHICS_API_D3D12)
-	void DecodeGetD3D12GraphicsRequirementsKHR(const FOpenXRAPIPacketBase& BasePacket);
+	bool DecodeGetD3D12GraphicsRequirementsKHR(const FOpenXRAPIPacketBase& BasePacket);
 #endif
 
 	FArrayReader EncodedData;
 
-	typedef void(FOpenXRCaptureDecoder::* ApiDecodeFn)(const FOpenXRAPIPacketBase& BasePacket);
+	typedef bool(FOpenXRCaptureDecoder::* ApiDecodeFn)(const FOpenXRAPIPacketBase& BasePacket);
 	TStaticArray<ApiDecodeFn, (uint32)EOpenXRAPIPacketId::NumValidAPIPacketIds> DecodeFnTable;
 
 	// derived state from capture
