XRScribe: do not error on a missing capture file to allow always-on usage.

diff --git a/Engine/Plugins/Runtime/XR/XRScribe/Source/XRScribe/Private/XRScribeAPISurface.cpp b/Engine/Plugins/Runtime/XR/XRScribe/Source/XRScribe/Private/XRScribeAPISurface.cpp
--- a/Engine/Plugins/Runtime/XR/XRScribe/Source/XRScribe/Private/XRScribeAPISurface.cpp
+++ b/Engine/Plugins/Runtime/XR/XRScribe/Source/XRScribe/Private/XRScribeAPISurface.cpp
@@ -793,7 +793,7 @@ public:
 			}
 			else
 			{
-				UE_LOG(LogXRScribeAPI, Warning, TEXT("Emulation Layer disabled"));
+				UE_LOG(LogXRScribeAPI, Log, TEXT("Emulation Layer disabled - unable to load the capture file."));
 			}
 
 			return bEmulationLoadedCapture;
diff --git a/Engine/Plugins/Runtime/XR/XRScribe/Source/XRScribe/Private/XRScribeEmulationLayer.cpp b/Engine/Plugins/Runtime/XR/XRScribe/Source/XRScribe/Private/XRScribeEmulationLayer.cpp
--- a/Engine/Plugins/Runtime/XR/XRScribe/Source/XRScribe/Private/XRScribeEmulationLayer.cpp
+++ b/Engine/Plugins/Runtime/XR/XRScribe/Source/XRScribe/Private/XRScribeEmulationLayer.cpp
@@ -293,7 +293,7 @@ bool FOpenXREmulationLayer::LoadCaptureFromFile(const FString& EmulationLoadPath
 	}
 	else
 	{
-		UE_LOG(LogXRScribeEmulate, Error, TEXT("Capture failed to load"));
+		UE_LOG(LogXRScribeEmulate, Log, TEXT("Capture '%s' failed to load."), *EmulationLoadPath);
 		return false;
 	}
 }
