Fix XRScribe issuing warnings when a capture file is not present.

diff --git a/Engine/Plugins/Runtime/XR/XRScribe/Source/XRScribe/Private/XRScribeEmulationLayer.cpp b/Engine/Plugins/Runtime/XR/XRScribe/Source/XRScribe/Private/XRScribeEmulationLayer.cpp
--- a/Engine/Plugins/Runtime/XR/XRScribe/Source/XRScribe/Private/XRScribeEmulationLayer.cpp
+++ b/Engine/Plugins/Runtime/XR/XRScribe/Source/XRScribe/Private/XRScribeEmulationLayer.cpp
@@ -264,7 +264,7 @@ bool FOpenXREmulationLayer::SupportsInstanceExtension(const ANSICHAR* ExtensionN
 
 bool FOpenXREmulationLayer::LoadCaptureFromFile(const FString& EmulationLoadPath)
 {
-	if (FFileHelper::LoadFileToArray(CaptureDecoder.GetEncodedData(), *EmulationLoadPath))
+	if (FFileHelper::LoadFileToArray(CaptureDecoder.GetEncodedData(), *EmulationLoadPath, EFileRead::FILEREAD_Silent))
 	{
 		UE_LOG(LogXRScribeEmulate, Log, TEXT("Capture successfully loaded: %s"), *EmulationLoadPath);
 		if (CaptureDecoder.DecodeDataFromMemory())
